AITestSuite module stub #UE4

- no real tests included just yet

diff --git a/Engine/Source/Developer/AITestSuite/AITestSuite.Build.cs b/Engine/Source/Developer/AITestSuite/AITestSuite.Build.cs
new file mode 100644
--- /dev/null
+++ b/Engine/Source/Developer/AITestSuite/AITestSuite.Build.cs
@@ -0,0 +1,47 @@
+// Copyright 1998-2014 Epic Games, Inc. All Rights Reserved.
+
+namespace UnrealBuildTool.Rules
+{
+	public class AITestSuite : ModuleRules
+	{
+        public AITestSuite(TargetInfo Target)
+		{
+			PublicIncludePaths.AddRange(
+				new string[] {
+					// ... add public include paths required here ...
+				}
+				);
+
+			PrivateIncludePaths.AddRange(
+				new string[] {
+					"Developer/AITestSuite/Private",
+					// ... add other private include paths required here ...
+				}
+				);
+
+			PublicDependencyModuleNames.AddRange(
+				new string[]
+				{
+					"Core",
+					"CoreUObject",
+					"Engine",
+                    "AIModule",
+				}
+				);
+
+			PrivateDependencyModuleNames.AddRange(
+				new string[]
+				{
+					// ... add private dependencies that you statically link with here ...
+				}
+				);
+
+			DynamicallyLoadedModuleNames.AddRange(
+				new string[]
+				{
+					// ... add any modules that your module loads dynamically here ...
+				}
+				);
+		}
+	}
+}
\ No newline at end of file
diff --git a/Engine/Source/Developer/AITestSuite/Private/AITestSuite.cpp b/Engine/Source/Developer/AITestSuite/Private/AITestSuite.cpp
new file mode 100644
--- /dev/null
+++ b/Engine/Source/Developer/AITestSuite/Private/AITestSuite.cpp
@@ -0,0 +1,30 @@
+// Copyright 1998-2014 Epic Games, Inc. All Rights Reserved.
+
+#include "AITestSuitePrivatePCH.h"
+
+
+class FAITestSuite : public IAITestSuite
+{
+	/** IModuleInterface implementation */
+	virtual void StartupModule() override;
+	virtual void ShutdownModule() override;
+};
+
+IMPLEMENT_MODULE(FAITestSuite, AITestSuite)
+
+
+
+void FAITestSuite::StartupModule()
+{
+	// This code will execute after your module is loaded into memory (but after global variables are initialized, of course.)
+}
+
+
+void FAITestSuite::ShutdownModule()
+{
+	// This function may be called during shutdown to clean up your module.  For modules that support dynamic reloading,
+	// we call this function before unloading the module.
+}
+
+
+
diff --git a/Engine/Source/Developer/AITestSuite/Private/AITestSuitePrivatePCH.h b/Engine/Source/Developer/AITestSuite/Private/AITestSuitePrivatePCH.h
new file mode 100644
--- /dev/null
+++ b/Engine/Source/Developer/AITestSuite/Private/AITestSuitePrivatePCH.h
@@ -0,0 +1,6 @@
+// Copyright 1998-2014 Epic Games, Inc. All Rights Reserved.
+
+#include "AITestSuite.h"
+
+// You should place include statements to your module's private header files here.  You only need to
+// add includes for headers that are used in most of your module's source files though.
diff --git a/Engine/Source/Developer/AITestSuite/Private/Tests/AITestsCommon.cpp b/Engine/Source/Developer/AITestSuite/Private/Tests/AITestsCommon.cpp
new file mode 100644
--- /dev/null
+++ b/Engine/Source/Developer/AITestSuite/Private/Tests/AITestsCommon.cpp
@@ -0,0 +1,4 @@
+// Copyright 1998-2014 Epic Games, Inc. All Rights Reserved.
+
+#include "AITestSuitePrivatePCH.h"
+
diff --git a/Engine/Source/Developer/AITestSuite/Private/Tests/AITestsCommon.h b/Engine/Source/Developer/AITestSuite/Private/Tests/AITestsCommon.h
new file mode 100644
--- /dev/null
+++ b/Engine/Source/Developer/AITestSuite/Private/Tests/AITestsCommon.h
@@ -0,0 +1,4 @@
+// Copyright 1998-2014 Epic Games, Inc. All Rights Reserved.
+
+#pragma once
+
diff --git a/Engine/Source/Developer/AITestSuite/Private/Tests/BBTest.cpp b/Engine/Source/Developer/AITestSuite/Private/Tests/BBTest.cpp
new file mode 100644
--- /dev/null
+++ b/Engine/Source/Developer/AITestSuite/Private/Tests/BBTest.cpp
@@ -0,0 +1,3 @@
+// Copyright 1998-2014 Epic Games, Inc. All Rights Reserved.
+
+#include "AITestSuitePrivatePCH.h"
\ No newline at end of file
diff --git a/Engine/Source/Developer/AITestSuite/Private/Tests/BTTest.cpp b/Engine/Source/Developer/AITestSuite/Private/Tests/BTTest.cpp
new file mode 100644
--- /dev/null
+++ b/Engine/Source/Developer/AITestSuite/Private/Tests/BTTest.cpp
@@ -0,0 +1,30 @@
+// Copyright 1998-2014 Epic Games, Inc. All Rights Reserved.
+
+#include "AITestSuitePrivatePCH.h"
+
+#define LOCTEXT_NAMESPACE "AITestSuite_BTTest"
+
+// create test base class
+IMPLEMENT_COMPLEX_AUTOMATION_TEST(FAITestSuite_BehaviorTree, "Engine.AI.Behavior Trees", (EAutomationTestFlags::ATF_Game | EAutomationTestFlags::ATF_Editor))
+
+/** 
+ * Requests a enumeration of all maps to be loaded
+ */
+void FAITestSuite_BehaviorTree::GetTests(TArray<FString>& OutBeautifiedNames, TArray<FString>& OutTestCommands) const
+{
+	OutBeautifiedNames.Add(TEXT("Simple Parallel aborting"));
+	OutTestCommands.Add(TEXT("BTTestID"));
+}
+
+/** 
+ * Execute the loading of each map and performance captures
+ *
+ * @param Parameters - Should specify which map name to load
+ * @return	TRUE if the test was successful, FALSE otherwise
+ */
+bool FAITestSuite_BehaviorTree::RunTest(const FString& Parameters)
+{	
+	return true;
+}
+
+#undef LOCTEXT_NAMESPACE
\ No newline at end of file
diff --git a/Engine/Source/Developer/AITestSuite/Private/Tests/PawnActionsTest.cpp b/Engine/Source/Developer/AITestSuite/Private/Tests/PawnActionsTest.cpp
new file mode 100644
--- /dev/null
+++ b/Engine/Source/Developer/AITestSuite/Private/Tests/PawnActionsTest.cpp
@@ -0,0 +1,3 @@
+// Copyright 1998-2014 Epic Games, Inc. All Rights Reserved.
+
+#include "AITestSuitePrivatePCH.h"
\ No newline at end of file
diff --git a/Engine/Source/Developer/AITestSuite/Public/AITestSuite.h b/Engine/Source/Developer/AITestSuite/Public/AITestSuite.h
new file mode 100644
--- /dev/null
+++ b/Engine/Source/Developer/AITestSuite/Public/AITestSuite.h
@@ -0,0 +1,37 @@
+// Copyright 1998-2014 Epic Games, Inc. All Rights Reserved.
+
+#pragma once
+
+#include "ModuleManager.h"
+
+
+/**
+ * The public interface to this module
+ */
+class IAITestSuite : public IModuleInterface
+{
+
+public:
+
+	/**
+	 * Singleton-like access to this module's interface.  This is just for convenience!
+	 * Beware of calling this during the shutdown phase, though.  Your module might have been unloaded already.
+	 *
+	 * @return Returns singleton instance, loading the module on demand if needed
+	 */
+	static inline IAITestSuite& Get()
+	{
+		return FModuleManager::LoadModuleChecked< IAITestSuite >("AITestSuite");
+	}
+
+	/**
+	 * Checks to see if this module is loaded and ready.  It is only valid to call Get() if IsAvailable() returns true.
+	 *
+	 * @return True if the module is loaded and ready to use
+	 */
+	static inline bool IsAvailable()
+	{
+		return FModuleManager::Get().IsModuleLoaded( "AITestSuite" );
+	}
+};
+
diff --git a/Engine/Source/Runtime/AIModule/Private/AIModule.cpp b/Engine/Source/Runtime/AIModule/Private/AIModule.cpp
--- a/Engine/Source/Runtime/AIModule/Private/AIModule.cpp
+++ b/Engine/Source/Runtime/AIModule/Private/AIModule.cpp
@@ -29,6 +29,7 @@ void FAIModule::StartupModule()
 	// This code will execute after your module is loaded into memory and after global variables initialization. We needs some place to load GameplayDebugger module so it's best place for it now.
 #if !(UE_BUILD_SHIPPING || UE_BUILD_TEST)
 	FModuleManager::LoadModulePtr< IModuleInterface >("GameplayDebugger");
+	FModuleManager::LoadModulePtr< IModuleInterface >("AITestSuite");
 #endif
 
 #if WITH_EDITOR && ENABLE_VISUAL_LOG
