added functional tests for behavior's tree component nodes

diff --git a/Engine/Source/Developer/AITestSuite/Classes/BehaviorTree/TestBTDecorator_CantExecute.h b/Engine/Source/Developer/AITestSuite/Classes/BehaviorTree/TestBTDecorator_CantExecute.h
new file mode 100644
--- /dev/null
+++ b/Engine/Source/Developer/AITestSuite/Classes/BehaviorTree/TestBTDecorator_CantExecute.h
@@ -0,0 +1,13 @@
+// Copyright 1998-2014 Epic Games, Inc. All Rights Reserved.
+
+#pragma once 
+#include "BehaviorTree/BTDecorator.h"
+#include "TestBTDecorator_CantExecute.generated.h"
+
+UCLASS()
+class UTestBTDecorator_CantExecute : public UBTDecorator
+{
+	GENERATED_UCLASS_BODY()
+
+	virtual bool CalculateRawConditionValue(class UBehaviorTreeComponent* OwnerComp, uint8* NodeMemory) const override;
+};
diff --git a/Engine/Source/Developer/AITestSuite/Classes/BehaviorTree/TestBTObject.h b/Engine/Source/Developer/AITestSuite/Classes/BehaviorTree/TestBTObject.h
new file mode 100644
--- /dev/null
+++ b/Engine/Source/Developer/AITestSuite/Classes/BehaviorTree/TestBTObject.h
@@ -0,0 +1,38 @@
+// Copyright 1998-2014 Epic Games, Inc. All Rights Reserved.
+#pragma once
+
+#include "Tickable.h"
+#include "BehaviorTree/BehaviorTreeComponent.h"
+#include "BehaviorTree/BlackboardComponent.h"
+#include "BehaviorTree/BehaviorTree.h"
+#include "TestBTObject.generated.h"
+
+struct FTestTickHelper : FTickableGameObject
+{
+	TWeakObjectPtr<class UTestBTObject> Owner;
+
+	FTestTickHelper() : Owner(NULL) {}
+	virtual void Tick(float DeltaTime);
+	virtual bool IsTickable() const { return Owner.IsValid(); }
+	virtual bool IsTickableInEditor() const { return true; }
+	virtual TStatId GetStatId() const;
+};
+
+UCLASS()
+class UTestBTObject : public UObject
+{
+	GENERATED_UCLASS_BODY()
+
+	UBehaviorTreeComponent* BTComp;
+	UBlackboardComponent* BBComp;
+	FTestTickHelper TickHelper;
+	
+	static TArray<int32> ExecutionLog;
+	TArray<int32> ExpectedResult;
+
+	void TickMe(float DeltaTime);
+	bool RunTest(UBehaviorTree* TreeOb);
+	bool IsRunning() const;
+
+	void VerifyResults();
+};
diff --git a/Engine/Source/Developer/AITestSuite/Classes/BehaviorTree/TestBTTask_Log.h b/Engine/Source/Developer/AITestSuite/Classes/BehaviorTree/TestBTTask_Log.h
new file mode 100644
--- /dev/null
+++ b/Engine/Source/Developer/AITestSuite/Classes/BehaviorTree/TestBTTask_Log.h
@@ -0,0 +1,34 @@
+// Copyright 1998-2014 Epic Games, Inc. All Rights Reserved.
+#pragma once
+
+#include "BehaviorTree/BTTaskNode.h"
+#include "TestBTTask_Log.generated.h"
+
+struct FBTLogTaskMemory
+{
+	/** time left */
+	float RemainingWaitTime;
+};
+
+UCLASS()
+class UTestBTTask_Log : public UBTTaskNode
+{
+	GENERATED_UCLASS_BODY()
+
+	UPROPERTY()
+	int32 LogIndex;
+
+	UPROPERTY()
+	float ExecutionTime;
+
+	UPROPERTY()
+	TEnumAsByte<EBTNodeResult::Type> LogResult;
+
+	virtual EBTNodeResult::Type ExecuteTask(class UBehaviorTreeComponent* OwnerComp, uint8* NodeMemory) override;
+	virtual uint16 GetInstanceMemorySize() const override;
+
+	void LogExecution(class UBehaviorTreeComponent* OwnerComp);
+
+protected:
+	virtual void TickTask(class UBehaviorTreeComponent* OwnerComp, uint8* NodeMemory, float DeltaSeconds) override;
+};
diff --git a/Engine/Source/Developer/AITestSuite/Classes/BehaviorTree/TestBTTask_SetFlag.h b/Engine/Source/Developer/AITestSuite/Classes/BehaviorTree/TestBTTask_SetFlag.h
new file mode 100644
--- /dev/null
+++ b/Engine/Source/Developer/AITestSuite/Classes/BehaviorTree/TestBTTask_SetFlag.h
@@ -0,0 +1,19 @@
+// Copyright 1998-2014 Epic Games, Inc. All Rights Reserved.
+#pragma once
+
+#include "BehaviorTree/BTTaskNode.h"
+#include "TestBTTask_SetFlag.generated.h"
+
+UCLASS()
+class UTestBTTask_SetFlag : public UBTTaskNode
+{
+	GENERATED_UCLASS_BODY()
+
+	UPROPERTY()
+	bool bValue;
+
+	UPROPERTY()
+	TEnumAsByte<EBTNodeResult::Type> TaskResult;
+
+	virtual EBTNodeResult::Type ExecuteTask(class UBehaviorTreeComponent* OwnerComp, uint8* NodeMemory) override;
+};
diff --git a/Engine/Source/Developer/AITestSuite/Private/BehaviorTree/TestBTDecorator_CantExecute.cpp b/Engine/Source/Developer/AITestSuite/Private/BehaviorTree/TestBTDecorator_CantExecute.cpp
new file mode 100644
--- /dev/null
+++ b/Engine/Source/Developer/AITestSuite/Private/BehaviorTree/TestBTDecorator_CantExecute.cpp
@@ -0,0 +1,18 @@
+// Copyright 1998-2014 Epic Games, Inc. All Rights Reserved.
+
+#include "AITestSuitePrivatePCH.h"
+#include "BehaviorTree/TestBTDecorator_CantExecute.h"
+
+UTestBTDecorator_CantExecute::UTestBTDecorator_CantExecute(const class FPostConstructInitializeProperties& PCIP) : Super(PCIP)
+{
+	NodeName = TEXT("Can't Exexcute");
+
+	bAllowAbortNone = false;
+	bAllowAbortLowerPri = false;
+	bAllowAbortChildNodes = false;
+}
+
+bool UTestBTDecorator_CantExecute::CalculateRawConditionValue(class UBehaviorTreeComponent* OwnerComp, uint8* NodeMemory) const
+{
+	return false;
+}
diff --git a/Engine/Source/Developer/AITestSuite/Private/BehaviorTree/TestBTObject.cpp b/Engine/Source/Developer/AITestSuite/Private/BehaviorTree/TestBTObject.cpp
new file mode 100644
--- /dev/null
+++ b/Engine/Source/Developer/AITestSuite/Private/BehaviorTree/TestBTObject.cpp
@@ -0,0 +1,98 @@
+// Copyright 1998-2014 Epic Games, Inc. All Rights Reserved.
+
+#include "AITestSuitePrivatePCH.h"
+#include "BehaviorTree/TestBTObject.h"
+
+DEFINE_LOG_CATEGORY_STATIC(LogBehaviorTreeTest, Log, All);
+
+void FTestTickHelper::Tick(float DeltaTime)
+{
+	if (Owner.IsValid())
+	{
+		Owner->TickMe(DeltaTime);
+	}
+}
+
+TStatId FTestTickHelper::GetStatId() const
+{
+	RETURN_QUICK_DECLARE_CYCLE_STAT(FTestTickHelper, STATGROUP_Tickables);
+}
+
+TArray<int32> UTestBTObject::ExecutionLog;
+
+UTestBTObject::UTestBTObject(const class FPostConstructInitializeProperties& PCIP) : Super(PCIP)
+{
+}
+
+void UTestBTObject::TickMe(float DeltaTime)
+{
+	if (BBComp)
+	{
+		BBComp->TickComponent(DeltaTime, ELevelTick::LEVELTICK_All, NULL);
+	}
+
+	if (BTComp)
+	{
+		BTComp->TickComponent(DeltaTime, ELevelTick::LEVELTICK_All, NULL);
+	}
+}
+
+bool UTestBTObject::RunTest(UBehaviorTree* TreeOb)
+{
+	const TIndirectArray<FWorldContext>& Contexts = GEngine->GetWorldContexts();
+	UWorld* MyWorld = NULL;
+	for (int32 Idx = 0; Idx < Contexts.Num(); Idx++)
+	{
+		if (Contexts[Idx].World())
+		{
+			MyWorld = Contexts[Idx].World();
+			break;
+		}
+	}
+
+	if (MyWorld == NULL)
+	{
+		return false;
+	}
+
+	BTComp = NewObject<UBehaviorTreeComponent>(MyWorld);
+	BBComp = NewObject<UBlackboardComponent>(MyWorld);
+
+	BBComp->InitializeBlackboard(TreeOb->BlackboardAsset);
+	BTComp->CacheBlackboardComponent(BBComp);
+
+	BBComp->RegisterComponentWithWorld(MyWorld);
+	BTComp->RegisterComponentWithWorld(MyWorld);
+
+	ExecutionLog.Reset();
+	BTComp->StartTree(TreeOb, EBTExecutionMode::SingleRun);
+	
+	TickHelper.Owner = this;
+	return true;
+}
+
+bool UTestBTObject::IsRunning() const
+{
+	return BTComp && BTComp->IsRunning() && BTComp->GetRootTree();
+}
+
+void UTestBTObject::VerifyResults()
+{
+	const bool bMatch = (ExpectedResult == ExecutionLog);
+	if (!bMatch)
+	{
+		FString Desc;
+		for (int32 Idx = 0; Idx < ExecutionLog.Num(); Idx++)
+		{
+			Desc += TTypeToString<int32>::ToString(ExecutionLog[Idx]);
+			if (Idx < (ExecutionLog.Num() - 1))
+			{
+				Desc += TEXT(", ");
+			}
+		}
+
+		UE_LOG(LogBehaviorTreeTest, Error, TEXT("Results are not matching!\nExecution log: %s"), *Desc);
+	}
+
+	TickHelper.Owner.Reset();
+}
diff --git a/Engine/Source/Developer/AITestSuite/Private/BehaviorTree/TestBTTask_Log.cpp b/Engine/Source/Developer/AITestSuite/Private/BehaviorTree/TestBTTask_Log.cpp
new file mode 100644
--- /dev/null
+++ b/Engine/Source/Developer/AITestSuite/Private/BehaviorTree/TestBTTask_Log.cpp
@@ -0,0 +1,51 @@
+// Copyright 1998-2014 Epic Games, Inc. All Rights Reserved.
+
+#include "AITestSuitePrivatePCH.h"
+#include "BehaviorTree/TestBTTask_Log.h"
+
+UTestBTTask_Log::UTestBTTask_Log(const class FPostConstructInitializeProperties& PCIP) : Super(PCIP)
+{
+	NodeName = "Log";
+	ExecutionTime = 5.0f;
+	LogIndex = 0;
+	LogResult = EBTNodeResult::Succeeded;
+
+	bNotifyTick = true;
+}
+
+EBTNodeResult::Type UTestBTTask_Log::ExecuteTask(class UBehaviorTreeComponent* OwnerComp, uint8* NodeMemory)
+{
+	FBTLogTaskMemory* MyMemory = (FBTLogTaskMemory*)NodeMemory;
+	MyMemory->RemainingWaitTime = ExecutionTime;
+
+	if (ExecutionTime <= 0.0f)
+	{
+		LogExecution(OwnerComp);
+		return LogResult;
+	}
+
+	return EBTNodeResult::InProgress;
+}
+
+void UTestBTTask_Log::TickTask(class UBehaviorTreeComponent* OwnerComp, uint8* NodeMemory, float DeltaSeconds)
+{
+	FBTLogTaskMemory* MyMemory = (FBTLogTaskMemory*)NodeMemory;
+	MyMemory->RemainingWaitTime -= DeltaSeconds;
+
+	if (MyMemory->RemainingWaitTime <= 0.0f)
+	{
+		// continue execution from this node
+		LogExecution(OwnerComp);
+		FinishLatentTask(OwnerComp, LogResult);
+	}
+}
+
+uint16 UTestBTTask_Log::GetInstanceMemorySize() const
+{
+	return sizeof(FBTLogTaskMemory);
+}
+
+void UTestBTTask_Log::LogExecution(class UBehaviorTreeComponent* OwnerComp)
+{
+	UTestBTObject::ExecutionLog.Add(LogIndex);
+}
diff --git a/Engine/Source/Developer/AITestSuite/Private/BehaviorTree/TestBTTask_SetFlag.cpp b/Engine/Source/Developer/AITestSuite/Private/BehaviorTree/TestBTTask_SetFlag.cpp
new file mode 100644
--- /dev/null
+++ b/Engine/Source/Developer/AITestSuite/Private/BehaviorTree/TestBTTask_SetFlag.cpp
@@ -0,0 +1,18 @@
+// Copyright 1998-2014 Epic Games, Inc. All Rights Reserved.
+
+#include "AITestSuitePrivatePCH.h"
+#include "BehaviorTree/TestBTTask_SetFlag.h"
+#include "BehaviorTree/BlackboardComponent.h"
+
+UTestBTTask_SetFlag::UTestBTTask_SetFlag(const class FPostConstructInitializeProperties& PCIP) : Super(PCIP)
+{
+	NodeName = "Log";
+	TaskResult = EBTNodeResult::Succeeded;
+	bValue = true;
+}
+
+EBTNodeResult::Type UTestBTTask_SetFlag::ExecuteTask(class UBehaviorTreeComponent* OwnerComp, uint8* NodeMemory)
+{
+	OwnerComp->GetBlackboardComponent()->SetValueAsBool(TEXT("Bool1"), bValue);
+	return TaskResult;
+}
diff --git a/Engine/Source/Developer/AITestSuite/Private/Tests/BTTest.cpp b/Engine/Source/Developer/AITestSuite/Private/Tests/BTTest.cpp
--- a/Engine/Source/Developer/AITestSuite/Private/Tests/BTTest.cpp
+++ b/Engine/Source/Developer/AITestSuite/Private/Tests/BTTest.cpp
@@ -1,9 +1,166 @@
 // Copyright 1998-2014 Epic Games, Inc. All Rights Reserved.
 
 #include "AITestSuitePrivatePCH.h"
+#include "BehaviorTree/BlackboardData.h"
+#include "BehaviorTree/Blackboard/BlackboardKeyAllTypes.h"
+#include "BehaviorTree/BehaviorTree.h"
+#include "BehaviorTree/Composites/BTComposite_Selector.h"
+#include "BehaviorTree/Composites/BTComposite_Sequence.h"
+#include "BehaviorTree/Composites/BTComposite_SimpleParallel.h"
+#include "BehaviorTree/Decorators//BTDecorator_ForceSuccess.h"
 
 #define LOCTEXT_NAMESPACE "AITestSuite_BTTest"
 
+// Helpers for creating behavior trees
+namespace BehaviorTreeHelpers
+{
+	UBehaviorTree* CreateBehaviorTree()
+	{
+		UBlackboardData* BB = NewObject<UBlackboardData>();
+		FBlackboardEntry KeyData;
+
+		KeyData.EntryName = TEXT("Bool1");
+		KeyData.KeyType = NewObject<UBlackboardKeyType_Bool>();
+		BB->Keys.Add(KeyData);
+
+		KeyData.EntryName = TEXT("Bool2");
+		KeyData.KeyType = NewObject<UBlackboardKeyType_Bool>();
+		BB->Keys.Add(KeyData);
+
+		UBehaviorTree* TreeOb = NewObject<UBehaviorTree>();
+		TreeOb->BlackboardAsset = BB;
+
+		return TreeOb;
+	}
+
+	UBTComposite_Selector* AddSelector(UBehaviorTree* TreeOb)
+	{
+		UBTComposite_Selector* NodeOb = NewObject<UBTComposite_Selector>(TreeOb);
+		NodeOb->InitializeFromAsset(TreeOb);
+		TreeOb->RootNode = NodeOb;
+		return NodeOb;
+	}
+
+	UBTComposite_Selector* AddSelector(UBTCompositeNode* ParentNode)
+	{
+		UBTComposite_Selector* NodeOb = NewObject<UBTComposite_Selector>(ParentNode->GetTreeAsset());
+		NodeOb->InitializeFromAsset(ParentNode->GetTreeAsset());
+		return NodeOb;
+	}
+
+	UBTComposite_Sequence* AddSequence(UBehaviorTree* TreeOb)
+	{
+		UBTComposite_Sequence* NodeOb = NewObject<UBTComposite_Sequence>(TreeOb);
+		NodeOb->InitializeFromAsset(TreeOb);
+		TreeOb->RootNode = NodeOb;
+		return NodeOb;
+	}
+
+	UBTComposite_Sequence* AddSequence(UBTCompositeNode* ParentNode)
+	{
+		UBTComposite_Sequence* NodeOb = NewObject<UBTComposite_Sequence>(ParentNode->GetTreeAsset());
+		NodeOb->InitializeFromAsset(ParentNode->GetTreeAsset());
+		return NodeOb;
+	}
+
+	UBTComposite_SimpleParallel* AddParallel(UBehaviorTree* TreeOb)
+	{
+		UBTComposite_SimpleParallel* NodeOb = NewObject<UBTComposite_SimpleParallel>(TreeOb);
+		NodeOb->InitializeFromAsset(TreeOb);
+		TreeOb->RootNode = NodeOb;
+		return NodeOb;
+	}
+
+	UBTComposite_SimpleParallel* AddParallel(UBTCompositeNode* ParentNode)
+	{
+		UBTComposite_SimpleParallel* NodeOb = NewObject<UBTComposite_SimpleParallel>(ParentNode->GetTreeAsset());
+		NodeOb->InitializeFromAsset(ParentNode->GetTreeAsset());
+		return NodeOb;
+	}
+
+	void AddTask(UBTCompositeNode* ParentNode, int32 LogIndex, EBTNodeResult::Type NodeResult, float ExecutionTime = 0.0f)
+	{
+		UTestBTTask_Log* TaskNode = NewObject<UTestBTTask_Log>(ParentNode->GetTreeAsset());
+		TaskNode->LogIndex = LogIndex;
+		TaskNode->LogResult = NodeResult;
+		TaskNode->ExecutionTime = ExecutionTime;
+
+		const int32 ChildIdx = ParentNode->Children.AddZeroed(1);
+		ParentNode->Children[ChildIdx].ChildTask = TaskNode;
+	}
+
+	template<class T>
+	T* WithDecorator(UBTCompositeNode* ParentNode, UClass* DecoratorClass = T::StaticClass())
+	{
+		T* DecoratorOb = NewObject<T>(ParentNode->GetTreeAsset());
+		ParentNode->Children.Last().Decorators.Add(DecoratorOb);
+
+		return DecoratorOb;
+	}
+
+	UBehaviorTree* CreateTestScenario_Selector(TArray<int32>& ExpectedResult)
+	{
+		UBehaviorTree* TreeOb = CreateBehaviorTree();
+		UBTCompositeNode* CompNode = AddSelector(TreeOb);
+		{
+			AddTask(CompNode, 0, EBTNodeResult::Failed);
+			
+			AddTask(CompNode, 1, EBTNodeResult::Succeeded);
+			{
+				WithDecorator<UTestBTDecorator_CantExecute>(CompNode);
+			}
+
+			AddTask(CompNode, 2, EBTNodeResult::Succeeded, 0.5f);
+			
+			AddTask(CompNode, 3, EBTNodeResult::Failed);
+		}
+
+		ExpectedResult.Add(0);
+		ExpectedResult.Add(2);
+		return TreeOb;
+	}
+
+	UBehaviorTree* CreateTestScenario_Sequence(TArray<int32>& ExpectedResult)
+	{
+		UBehaviorTree* TreeOb = CreateBehaviorTree();
+		UBTCompositeNode* CompNode = AddSequence(TreeOb);
+		{
+			AddTask(CompNode, 0, EBTNodeResult::Succeeded);
+
+			AddTask(CompNode, 1, EBTNodeResult::Failed);
+			{
+				WithDecorator<UTestBTDecorator_CantExecute>(CompNode);
+				WithDecorator<UBTDecorator_ForceSuccess>(CompNode);
+			}
+
+			AddTask(CompNode, 2, EBTNodeResult::Failed, 0.5f);
+			
+			AddTask(CompNode, 3, EBTNodeResult::Succeeded);
+		}
+
+		ExpectedResult.Add(0);
+		ExpectedResult.Add(2);
+		return TreeOb;
+	}
+}
+
+typedef UBehaviorTree* (*CreateTreeFuncPtr)(TArray<int32>& ExpectedResult);
+const int32 MaxBehaviorTreeTests = 32;
+CreateTreeFuncPtr BehaviorTreeTestFunc[MaxBehaviorTreeTests];
+
+DEFINE_LATENT_AUTOMATION_COMMAND_ONE_PARAMETER(FAITestSuite_WaitForBehaviorTreeCommand, UTestBTObject*, TestObject);
+bool FAITestSuite_WaitForBehaviorTreeCommand::Update()
+{
+	if (TestObject->IsRunning())
+	{
+		return false;
+	}
+
+	TestObject->VerifyResults();
+	TestObject->RemoveFromRoot();
+	return true;
+}
+
 // create test base class
 IMPLEMENT_COMPLEX_AUTOMATION_TEST(FAITestSuite_BehaviorTree, "Engine.AI.Behavior Trees", (EAutomationTestFlags::ATF_Game | EAutomationTestFlags::ATF_Editor))
 
@@ -12,8 +169,13 @@ IMPLEMENT_COMPLEX_AUTOMATION_TEST(FAITestSuite_BehaviorTree, "Engine.AI.Behavior
  */
 void FAITestSuite_BehaviorTree::GetTests(TArray<FString>& OutBeautifiedNames, TArray<FString>& OutTestCommands) const
 {
-	OutBeautifiedNames.Add(TEXT("Simple Parallel aborting"));
-	OutTestCommands.Add(TEXT("BTTestID"));
+	OutBeautifiedNames.Add(TEXT("Composite node: selector"));
+	OutTestCommands.Add(TEXT("0"));
+	BehaviorTreeTestFunc[0] = BehaviorTreeHelpers::CreateTestScenario_Selector;
+
+	OutBeautifiedNames.Add(TEXT("Composite node: sequence"));
+	OutTestCommands.Add(TEXT("1"));
+	BehaviorTreeTestFunc[1] = BehaviorTreeHelpers::CreateTestScenario_Sequence;
 }
 
 /** 
@@ -23,8 +185,27 @@ void FAITestSuite_BehaviorTree::GetTests(TArray<FString>& OutBeautifiedNames, TA
  * @return	TRUE if the test was successful, FALSE otherwise
  */
 bool FAITestSuite_BehaviorTree::RunTest(const FString& Parameters)
-{	
+{
+	int32 TestNum = 0;
+	TTypeFromString<int32>::FromString(TestNum, *Parameters);
+
+	if (TestNum < 0 || TestNum >= MaxBehaviorTreeTests)
+	{
+		return false;
+	}
+
+	UTestBTObject* TestObject = NewObject<UTestBTObject>();
+	TestObject->AddToRoot();
+
+	UBehaviorTree* TreeOb = BehaviorTreeTestFunc[TestNum](TestObject->ExpectedResult);
+	const bool bTestStarted = TestObject->RunTest(TreeOb);
+	if (!bTestStarted)
+	{
+		return false;
+	}
+
+	ADD_LATENT_AUTOMATION_COMMAND(FAITestSuite_WaitForBehaviorTreeCommand(TestObject));
 	return true;
 }
 
-#undef LOCTEXT_NAMESPACE
\ No newline at end of file
+#undef LOCTEXT_NAMESPACE
diff --git a/Engine/Source/Runtime/AIModule/Classes/BehaviorTree/Composites/BTComposite_Selector.h b/Engine/Source/Runtime/AIModule/Classes/BehaviorTree/Composites/BTComposite_Selector.h
--- a/Engine/Source/Runtime/AIModule/Classes/BehaviorTree/Composites/BTComposite_Selector.h
+++ b/Engine/Source/Runtime/AIModule/Classes/BehaviorTree/Composites/BTComposite_Selector.h
@@ -8,7 +8,7 @@
  * Selector composite node.
  * A selector node runs each child in turn until the first child node succeeds, in which case the selector fails.
  */
-UCLASS()
+UCLASS(MinimalAPI)
 class UBTComposite_Selector : public UBTCompositeNode
 {
 	GENERATED_UCLASS_BODY()
diff --git a/Engine/Source/Runtime/AIModule/Classes/BehaviorTree/Composites/BTComposite_Sequence.h b/Engine/Source/Runtime/AIModule/Classes/BehaviorTree/Composites/BTComposite_Sequence.h
--- a/Engine/Source/Runtime/AIModule/Classes/BehaviorTree/Composites/BTComposite_Sequence.h
+++ b/Engine/Source/Runtime/AIModule/Classes/BehaviorTree/Composites/BTComposite_Sequence.h
@@ -8,7 +8,7 @@
  * Sequence composite node.
  * A sequence node runs each child in turn until one fails, in which case the sequence succeeds.
  */
-UCLASS()
+UCLASS(MinimalAPI)
 class UBTComposite_Sequence : public UBTCompositeNode
 {
 	GENERATED_UCLASS_BODY()
diff --git a/Engine/Source/Runtime/AIModule/Classes/BehaviorTree/Decorators/BTDecorator_ForceSuccess.h b/Engine/Source/Runtime/AIModule/Classes/BehaviorTree/Decorators/BTDecorator_ForceSuccess.h
--- a/Engine/Source/Runtime/AIModule/Classes/BehaviorTree/Decorators/BTDecorator_ForceSuccess.h
+++ b/Engine/Source/Runtime/AIModule/Classes/BehaviorTree/Decorators/BTDecorator_ForceSuccess.h
@@ -13,7 +13,7 @@
  * - selector = would allow executing multiple nodes, turning it into a sequence...
  */
 
-UCLASS(HideCategories=(Condition))
+UCLASS(MinimalAPI, HideCategories = (Condition))
 class UBTDecorator_ForceSuccess : public UBTDecorator
 {
 	GENERATED_UCLASS_BODY()
diff --git a/Engine/Source/Runtime/AIModule/Private/BehaviorTree/BehaviorTreeComponent.cpp b/Engine/Source/Runtime/AIModule/Private/BehaviorTree/BehaviorTreeComponent.cpp
--- a/Engine/Source/Runtime/AIModule/Private/BehaviorTree/BehaviorTreeComponent.cpp
+++ b/Engine/Source/Runtime/AIModule/Private/BehaviorTree/BehaviorTreeComponent.cpp
@@ -168,10 +168,9 @@ void UBehaviorTreeComponent::StopTree()
 	PendingExecution = FBTPendingExecutionInfo();
 	ActiveInstanceIdx = 0;
 
-	if (IsRegistered())
+	if (IsRegistered() && GetWorld())
 	{
-		check(GetOwner() != NULL);
-		GetOwner()->GetWorldTimerManager().ClearTimer(this, &UBehaviorTreeComponent::ProcessExecutionRequest);
+		GetWorld()->GetTimerManager().ClearTimer(this, &UBehaviorTreeComponent::ProcessExecutionRequest);
 	}
 
 	// make sure to allow new execution requests since we just removed last request timer
@@ -519,7 +518,7 @@ void UBehaviorTreeComponent::ScheduleExecutionUpdate()
 	{
 		bRequestedFlowUpdate = true;
 
-		GetOwner()->GetWorldTimerManager().SetTimer(this, &UBehaviorTreeComponent::ProcessExecutionRequest, 0.001f, false);
+		GetWorld()->GetTimerManager().SetTimer(this, &UBehaviorTreeComponent::ProcessExecutionRequest, 0.001f, false);
 	}
 }
 
@@ -532,7 +531,7 @@ void UBehaviorTreeComponent::RequestExecution(class UBTCompositeNode* RequestedO
 		*UBehaviorTreeTypes::DescribeNodeHelper(RequestedBy),
 		*UBehaviorTreeTypes::DescribeNodeResult(ContinueWithResult));
 
-	if (!bIsRunning || GetOwner() == NULL || GetOwner()->IsPendingKillPending())
+	if (!bIsRunning || (GetOwner() && GetOwner()->IsPendingKillPending()))
 	{
 		UE_VLOG(GetOwner(), LogBehaviorTree, Log, TEXT("> skip: tree is not running"));
 		return;
