Back out changelist 2307603

- Fixing CIS, that changelist broke Mac/Linux builds.

diff --git a/Engine/Source/Developer/AITestSuite/Classes/BehaviorTree/TestBTDecorator_CantExecute.h b/Engine/Source/Developer/AITestSuite/Classes/BehaviorTree/TestBTDecorator_CantExecute.h
deleted file mode 100644
diff --git a/Engine/Source/Developer/AITestSuite/Classes/BehaviorTree/TestBTObject.h b/Engine/Source/Developer/AITestSuite/Classes/BehaviorTree/TestBTObject.h
deleted file mode 100644
diff --git a/Engine/Source/Developer/AITestSuite/Classes/BehaviorTree/TestBTTask_Log.h b/Engine/Source/Developer/AITestSuite/Classes/BehaviorTree/TestBTTask_Log.h
deleted file mode 100644
diff --git a/Engine/Source/Developer/AITestSuite/Classes/BehaviorTree/TestBTTask_SetFlag.h b/Engine/Source/Developer/AITestSuite/Classes/BehaviorTree/TestBTTask_SetFlag.h
deleted file mode 100644
diff --git a/Engine/Source/Developer/AITestSuite/Private/BehaviorTree/TestBTDecorator_CantExecute.cpp b/Engine/Source/Developer/AITestSuite/Private/BehaviorTree/TestBTDecorator_CantExecute.cpp
deleted file mode 100644
diff --git a/Engine/Source/Developer/AITestSuite/Private/BehaviorTree/TestBTObject.cpp b/Engine/Source/Developer/AITestSuite/Private/BehaviorTree/TestBTObject.cpp
deleted file mode 100644
diff --git a/Engine/Source/Developer/AITestSuite/Private/BehaviorTree/TestBTTask_Log.cpp b/Engine/Source/Developer/AITestSuite/Private/BehaviorTree/TestBTTask_Log.cpp
deleted file mode 100644
diff --git a/Engine/Source/Developer/AITestSuite/Private/BehaviorTree/TestBTTask_SetFlag.cpp b/Engine/Source/Developer/AITestSuite/Private/BehaviorTree/TestBTTask_SetFlag.cpp
deleted file mode 100644
diff --git a/Engine/Source/Developer/AITestSuite/Private/Tests/BTTest.cpp b/Engine/Source/Developer/AITestSuite/Private/Tests/BTTest.cpp
--- a/Engine/Source/Developer/AITestSuite/Private/Tests/BTTest.cpp
+++ b/Engine/Source/Developer/AITestSuite/Private/Tests/BTTest.cpp
@@ -1,166 +1,9 @@
 // Copyright 1998-2014 Epic Games, Inc. All Rights Reserved.
 
 #include "AITestSuitePrivatePCH.h"
-#include "BehaviorTree/BlackboardData.h"
-#include "BehaviorTree/Blackboard/BlackboardKeyAllTypes.h"
-#include "BehaviorTree/BehaviorTree.h"
-#include "BehaviorTree/Composites/BTComposite_Selector.h"
-#include "BehaviorTree/Composites/BTComposite_Sequence.h"
-#include "BehaviorTree/Composites/BTComposite_SimpleParallel.h"
-#include "BehaviorTree/Decorators//BTDecorator_ForceSuccess.h"
 
 #define LOCTEXT_NAMESPACE "AITestSuite_BTTest"
 
-// Helpers for creating behavior trees
-namespace BehaviorTreeHelpers
-{
-	UBehaviorTree* CreateBehaviorTree()
-	{
-		UBlackboardData* BB = NewObject<UBlackboardData>();
-		FBlackboardEntry KeyData;
-
-		KeyData.EntryName = TEXT("Bool1");
-		KeyData.KeyType = NewObject<UBlackboardKeyType_Bool>();
-		BB->Keys.Add(KeyData);
-
-		KeyData.EntryName = TEXT("Bool2");
-		KeyData.KeyType = NewObject<UBlackboardKeyType_Bool>();
-		BB->Keys.Add(KeyData);
-
-		UBehaviorTree* TreeOb = NewObject<UBehaviorTree>();
-		TreeOb->BlackboardAsset = BB;
-
-		return TreeOb;
-	}
-
-	UBTComposite_Selector* AddSelector(UBehaviorTree* TreeOb)
-	{
-		UBTComposite_Selector* NodeOb = NewObject<UBTComposite_Selector>(TreeOb);
-		NodeOb->InitializeFromAsset(TreeOb);
-		TreeOb->RootNode = NodeOb;
-		return NodeOb;
-	}
-
-	UBTComposite_Selector* AddSelector(UBTCompositeNode* ParentNode)
-	{
-		UBTComposite_Selector* NodeOb = NewObject<UBTComposite_Selector>(ParentNode->GetTreeAsset());
-		NodeOb->InitializeFromAsset(ParentNode->GetTreeAsset());
-		return NodeOb;
-	}
-
-	UBTComposite_Sequence* AddSequence(UBehaviorTree* TreeOb)
-	{
-		UBTComposite_Sequence* NodeOb = NewObject<UBTComposite_Sequence>(TreeOb);
-		NodeOb->InitializeFromAsset(TreeOb);
-		TreeOb->RootNode = NodeOb;
-		return NodeOb;
-	}
-
-	UBTComposite_Sequence* AddSequence(UBTCompositeNode* ParentNode)
-	{
-		UBTComposite_Sequence* NodeOb = NewObject<UBTComposite_Sequence>(ParentNode->GetTreeAsset());
-		NodeOb->InitializeFromAsset(ParentNode->GetTreeAsset());
-		return NodeOb;
-	}
-
-	UBTComposite_SimpleParallel* AddParallel(UBehaviorTree* TreeOb)
-	{
-		UBTComposite_SimpleParallel* NodeOb = NewObject<UBTComposite_SimpleParallel>(TreeOb);
-		NodeOb->InitializeFromAsset(TreeOb);
-		TreeOb->RootNode = NodeOb;
-		return NodeOb;
-	}
-
-	UBTComposite_SimpleParallel* AddParallel(UBTCompositeNode* ParentNode)
-	{
-		UBTComposite_SimpleParallel* NodeOb = NewObject<UBTComposite_SimpleParallel>(ParentNode->GetTreeAsset());
-		NodeOb->InitializeFromAsset(ParentNode->GetTreeAsset());
-		return NodeOb;
-	}
-
-	void AddTask(UBTCompositeNode* ParentNode, int32 LogIndex, EBTNodeResult::Type NodeResult, float ExecutionTime = 0.0f)
-	{
-		UTestBTTask_Log* TaskNode = NewObject<UTestBTTask_Log>(ParentNode->GetTreeAsset());
-		TaskNode->LogIndex = LogIndex;
-		TaskNode->LogResult = NodeResult;
-		TaskNode->ExecutionTime = ExecutionTime;
-
-		const int32 ChildIdx = ParentNode->Children.AddZeroed(1);
-		ParentNode->Children[ChildIdx].ChildTask = TaskNode;
-	}
-
-	template<class T>
-	T* WithDecorator(UBTCompositeNode* ParentNode, UClass* DecoratorClass = T::StaticClass())
-	{
-		T* DecoratorOb = NewObject<T>(ParentNode->GetTreeAsset());
-		ParentNode->Children.Last().Decorators.Add(DecoratorOb);
-
-		return DecoratorOb;
-	}
-
-	UBehaviorTree* CreateTestScenario_Selector(TArray<int32>& ExpectedResult)
-	{
-		UBehaviorTree* TreeOb = CreateBehaviorTree();
-		UBTCompositeNode* CompNode = AddSelector(TreeOb);
-		{
-			AddTask(CompNode, 0, EBTNodeResult::Failed);
-			
-			AddTask(CompNode, 1, EBTNodeResult::Succeeded);
-			{
-				WithDecorator<UTestBTDecorator_CantExecute>(CompNode);
-			}
-
-			AddTask(CompNode, 2, EBTNodeResult::Succeeded, 0.5f);
-			
-			AddTask(CompNode, 3, EBTNodeResult::Failed);
-		}
-
-		ExpectedResult.Add(0);
-		ExpectedResult.Add(2);
-		return TreeOb;
-	}
-
-	UBehaviorTree* CreateTestScenario_Sequence(TArray<int32>& ExpectedResult)
-	{
-		UBehaviorTree* TreeOb = CreateBehaviorTree();
-		UBTCompositeNode* CompNode = AddSequence(TreeOb);
-		{
-			AddTask(CompNode, 0, EBTNodeResult::Succeeded);
-
-			AddTask(CompNode, 1, EBTNodeResult::Failed);
-			{
-				WithDecorator<UTestBTDecorator_CantExecute>(CompNode);
-				WithDecorator<UBTDecorator_ForceSuccess>(CompNode);
-			}
-
-			AddTask(CompNode, 2, EBTNodeResult::Failed, 0.5f);
-			
-			AddTask(CompNode, 3, EBTNodeResult::Succeeded);
-		}
-
-		ExpectedResult.Add(0);
-		ExpectedResult.Add(2);
-		return TreeOb;
-	}
-}
-
-typedef UBehaviorTree* (*CreateTreeFuncPtr)(TArray<int32>& ExpectedResult);
-const int32 MaxBehaviorTreeTests = 32;
-CreateTreeFuncPtr BehaviorTreeTestFunc[MaxBehaviorTreeTests];
-
-DEFINE_LATENT_AUTOMATION_COMMAND_ONE_PARAMETER(FAITestSuite_WaitForBehaviorTreeCommand, UTestBTObject*, TestObject);
-bool FAITestSuite_WaitForBehaviorTreeCommand::Update()
-{
-	if (TestObject->IsRunning())
-	{
-		return false;
-	}
-
-	TestObject->VerifyResults();
-	TestObject->RemoveFromRoot();
-	return true;
-}
-
 // create test base class
 IMPLEMENT_COMPLEX_AUTOMATION_TEST(FAITestSuite_BehaviorTree, "Engine.AI.Behavior Trees", (EAutomationTestFlags::ATF_Game | EAutomationTestFlags::ATF_Editor))
 
@@ -169,13 +12,8 @@ IMPLEMENT_COMPLEX_AUTOMATION_TEST(FAITestSuite_BehaviorTree, "Engine.AI.Behavior
  */
 void FAITestSuite_BehaviorTree::GetTests(TArray<FString>& OutBeautifiedNames, TArray<FString>& OutTestCommands) const
 {
-	OutBeautifiedNames.Add(TEXT("Composite node: selector"));
-	OutTestCommands.Add(TEXT("0"));
-	BehaviorTreeTestFunc[0] = BehaviorTreeHelpers::CreateTestScenario_Selector;
-
-	OutBeautifiedNames.Add(TEXT("Composite node: sequence"));
-	OutTestCommands.Add(TEXT("1"));
-	BehaviorTreeTestFunc[1] = BehaviorTreeHelpers::CreateTestScenario_Sequence;
+	OutBeautifiedNames.Add(TEXT("Simple Parallel aborting"));
+	OutTestCommands.Add(TEXT("BTTestID"));
 }
 
 /** 
@@ -185,27 +23,8 @@ void FAITestSuite_BehaviorTree::GetTests(TArray<FString>& OutBeautifiedNames, TA
  * @return	TRUE if the test was successful, FALSE otherwise
  */
 bool FAITestSuite_BehaviorTree::RunTest(const FString& Parameters)
-{
-	int32 TestNum = 0;
-	TTypeFromString<int32>::FromString(TestNum, *Parameters);
-
-	if (TestNum < 0 || TestNum >= MaxBehaviorTreeTests)
-	{
-		return false;
-	}
-
-	UTestBTObject* TestObject = NewObject<UTestBTObject>();
-	TestObject->AddToRoot();
-
-	UBehaviorTree* TreeOb = BehaviorTreeTestFunc[TestNum](TestObject->ExpectedResult);
-	const bool bTestStarted = TestObject->RunTest(TreeOb);
-	if (!bTestStarted)
-	{
-		return false;
-	}
-
-	ADD_LATENT_AUTOMATION_COMMAND(FAITestSuite_WaitForBehaviorTreeCommand(TestObject));
+{	
 	return true;
 }
 
-#undef LOCTEXT_NAMESPACE
+#undef LOCTEXT_NAMESPACE
\ No newline at end of file
diff --git a/Engine/Source/Runtime/AIModule/Classes/BehaviorTree/Composites/BTComposite_Selector.h b/Engine/Source/Runtime/AIModule/Classes/BehaviorTree/Composites/BTComposite_Selector.h
--- a/Engine/Source/Runtime/AIModule/Classes/BehaviorTree/Composites/BTComposite_Selector.h
+++ b/Engine/Source/Runtime/AIModule/Classes/BehaviorTree/Composites/BTComposite_Selector.h
@@ -8,7 +8,7 @@
  * Selector composite node.
  * A selector node runs each child in turn until the first child node succeeds, in which case the selector fails.
  */
-UCLASS(MinimalAPI)
+UCLASS()
 class UBTComposite_Selector : public UBTCompositeNode
 {
 	GENERATED_UCLASS_BODY()
diff --git a/Engine/Source/Runtime/AIModule/Classes/BehaviorTree/Composites/BTComposite_Sequence.h b/Engine/Source/Runtime/AIModule/Classes/BehaviorTree/Composites/BTComposite_Sequence.h
--- a/Engine/Source/Runtime/AIModule/Classes/BehaviorTree/Composites/BTComposite_Sequence.h
+++ b/Engine/Source/Runtime/AIModule/Classes/BehaviorTree/Composites/BTComposite_Sequence.h
@@ -8,7 +8,7 @@
  * Sequence composite node.
  * A sequence node runs each child in turn until one fails, in which case the sequence succeeds.
  */
-UCLASS(MinimalAPI)
+UCLASS()
 class UBTComposite_Sequence : public UBTCompositeNode
 {
 	GENERATED_UCLASS_BODY()
diff --git a/Engine/Source/Runtime/AIModule/Classes/BehaviorTree/Decorators/BTDecorator_ForceSuccess.h b/Engine/Source/Runtime/AIModule/Classes/BehaviorTree/Decorators/BTDecorator_ForceSuccess.h
--- a/Engine/Source/Runtime/AIModule/Classes/BehaviorTree/Decorators/BTDecorator_ForceSuccess.h
+++ b/Engine/Source/Runtime/AIModule/Classes/BehaviorTree/Decorators/BTDecorator_ForceSuccess.h
@@ -13,7 +13,7 @@
  * - selector = would allow executing multiple nodes, turning it into a sequence...
  */
 
-UCLASS(MinimalAPI, HideCategories = (Condition))
+UCLASS(HideCategories=(Condition))
 class UBTDecorator_ForceSuccess : public UBTDecorator
 {
 	GENERATED_UCLASS_BODY()
diff --git a/Engine/Source/Runtime/AIModule/Private/BehaviorTree/BehaviorTreeComponent.cpp b/Engine/Source/Runtime/AIModule/Private/BehaviorTree/BehaviorTreeComponent.cpp
--- a/Engine/Source/Runtime/AIModule/Private/BehaviorTree/BehaviorTreeComponent.cpp
+++ b/Engine/Source/Runtime/AIModule/Private/BehaviorTree/BehaviorTreeComponent.cpp
@@ -168,9 +168,10 @@ void UBehaviorTreeComponent::StopTree()
 	PendingExecution = FBTPendingExecutionInfo();
 	ActiveInstanceIdx = 0;
 
-	if (IsRegistered() && GetWorld())
+	if (IsRegistered())
 	{
-		GetWorld()->GetTimerManager().ClearTimer(this, &UBehaviorTreeComponent::ProcessExecutionRequest);
+		check(GetOwner() != NULL);
+		GetOwner()->GetWorldTimerManager().ClearTimer(this, &UBehaviorTreeComponent::ProcessExecutionRequest);
 	}
 
 	// make sure to allow new execution requests since we just removed last request timer
@@ -518,7 +519,7 @@ void UBehaviorTreeComponent::ScheduleExecutionUpdate()
 	{
 		bRequestedFlowUpdate = true;
 
-		GetWorld()->GetTimerManager().SetTimer(this, &UBehaviorTreeComponent::ProcessExecutionRequest, 0.001f, false);
+		GetOwner()->GetWorldTimerManager().SetTimer(this, &UBehaviorTreeComponent::ProcessExecutionRequest, 0.001f, false);
 	}
 }
 
@@ -531,7 +532,7 @@ void UBehaviorTreeComponent::RequestExecution(class UBTCompositeNode* RequestedO
 		*UBehaviorTreeTypes::DescribeNodeHelper(RequestedBy),
 		*UBehaviorTreeTypes::DescribeNodeResult(ContinueWithResult));
 
-	if (!bIsRunning || (GetOwner() && GetOwner()->IsPendingKillPending()))
+	if (!bIsRunning || GetOwner() == NULL || GetOwner()->IsPendingKillPending())
 	{
 		UE_VLOG(GetOwner(), LogBehaviorTree, Log, TEXT("> skip: tree is not running"));
 		return;
