More structure to AITestSuite #UE4

diff --git a/Engine/Source/Developer/AITestSuite/Classes/BehaviorTree/TestBTDecorator_CantExecute.h b/Engine/Source/Developer/AITestSuite/Classes/BehaviorTree/TestBTDecorator_CantExecute.h
new file mode 100644
--- /dev/null
+++ b/Engine/Source/Developer/AITestSuite/Classes/BehaviorTree/TestBTDecorator_CantExecute.h
@@ -0,0 +1,13 @@
+// Copyright 1998-2014 Epic Games, Inc. All Rights Reserved.
+
+#pragma once 
+#include "BehaviorTree/BTDecorator.h"
+#include "TestBTDecorator_CantExecute.generated.h"
+
+UCLASS()
+class UTestBTDecorator_CantExecute : public UBTDecorator
+{
+	GENERATED_UCLASS_BODY()
+
+	virtual bool CalculateRawConditionValue(class UBehaviorTreeComponent* OwnerComp, uint8* NodeMemory) const override;
+};
diff --git a/Engine/Source/Developer/AITestSuite/Classes/BehaviorTree/TestBTTask_Log.h b/Engine/Source/Developer/AITestSuite/Classes/BehaviorTree/TestBTTask_Log.h
new file mode 100644
--- /dev/null
+++ b/Engine/Source/Developer/AITestSuite/Classes/BehaviorTree/TestBTTask_Log.h
@@ -0,0 +1,34 @@
+// Copyright 1998-2014 Epic Games, Inc. All Rights Reserved.
+#pragma once
+
+#include "BehaviorTree/BTTaskNode.h"
+#include "TestBTTask_Log.generated.h"
+
+struct FBTLogTaskMemory
+{
+	/** time left */
+	float RemainingWaitTime;
+};
+
+UCLASS()
+class UTestBTTask_Log : public UBTTaskNode
+{
+	GENERATED_UCLASS_BODY()
+
+	UPROPERTY()
+	int32 LogIndex;
+
+	UPROPERTY()
+	float ExecutionTime;
+
+	UPROPERTY()
+	TEnumAsByte<EBTNodeResult::Type> LogResult;
+
+	virtual EBTNodeResult::Type ExecuteTask(class UBehaviorTreeComponent* OwnerComp, uint8* NodeMemory) override;
+	virtual uint16 GetInstanceMemorySize() const override;
+
+	void LogExecution(class UBehaviorTreeComponent* OwnerComp);
+
+protected:
+	virtual void TickTask(class UBehaviorTreeComponent* OwnerComp, uint8* NodeMemory, float DeltaSeconds) override;
+};
diff --git a/Engine/Source/Developer/AITestSuite/Classes/BehaviorTree/TestBTTask_SetFlag.h b/Engine/Source/Developer/AITestSuite/Classes/BehaviorTree/TestBTTask_SetFlag.h
new file mode 100644
--- /dev/null
+++ b/Engine/Source/Developer/AITestSuite/Classes/BehaviorTree/TestBTTask_SetFlag.h
@@ -0,0 +1,19 @@
+// Copyright 1998-2014 Epic Games, Inc. All Rights Reserved.
+#pragma once
+
+#include "BehaviorTree/BTTaskNode.h"
+#include "TestBTTask_SetFlag.generated.h"
+
+UCLASS()
+class UTestBTTask_SetFlag : public UBTTaskNode
+{
+	GENERATED_UCLASS_BODY()
+
+	UPROPERTY()
+	bool bValue;
+
+	UPROPERTY()
+	TEnumAsByte<EBTNodeResult::Type> TaskResult;
+
+	virtual EBTNodeResult::Type ExecuteTask(class UBehaviorTreeComponent* OwnerComp, uint8* NodeMemory) override;
+};
diff --git a/Engine/Source/Developer/AITestSuite/Classes/MockAI.h b/Engine/Source/Developer/AITestSuite/Classes/MockAI.h
new file mode 100644
--- /dev/null
+++ b/Engine/Source/Developer/AITestSuite/Classes/MockAI.h
@@ -0,0 +1,54 @@
+// Copyright 1998-2014 Epic Games, Inc. All Rights Reserved.
+
+#pragma once
+
+#include "Tickable.h"
+
+#include "MockAI.generated.h"
+
+struct FTestTickHelper : FTickableGameObject
+{
+	TWeakObjectPtr<class UMockAI> Owner;
+
+	FTestTickHelper() : Owner(NULL) {}
+	virtual void Tick(float DeltaTime);
+	virtual bool IsTickable() const { return Owner.IsValid(); }
+	virtual bool IsTickableInEditor() const { return true; }
+	virtual TStatId GetStatId() const;
+};
+
+UCLASS()
+class UMockAI : public UObject
+{
+	GENERATED_UCLASS_BODY()
+
+	virtual ~UMockAI();
+
+	FTestTickHelper TickHelper;
+
+	UPROPERTY()
+	UBlackboardComponent* BBComp;
+
+	UPROPERTY()
+	UBrainComponent* BrainComp;
+
+	UPROPERTY()
+	UAIPerceptionComponent* PerceptionComp;
+
+	UPROPERTY()
+	UPawnActionsComponent* PawnActionComp;
+	
+	template<typename TBrainClass>
+	void UseBrainComponent()
+	{
+		BrainComp = NewObject<TBrainClass>(FAITestHelpers::GetWorld());
+	}
+
+	void UseBlackboardComponent();
+	void UsePerceptionComponent();
+	void UsePawnActionsComponent();
+
+	void SetEnableTicking(bool bShouldTick);
+
+	virtual void TickMe(float DeltaTime);
+};
diff --git a/Engine/Source/Developer/AITestSuite/Classes/MockAI_BT.h b/Engine/Source/Developer/AITestSuite/Classes/MockAI_BT.h
new file mode 100644
--- /dev/null
+++ b/Engine/Source/Developer/AITestSuite/Classes/MockAI_BT.h
@@ -0,0 +1,20 @@
+// Copyright 1998-2014 Epic Games, Inc. All Rights Reserved.
+#pragma once
+
+#include "MockAI.h"
+#include "MockAI_BT.generated.h"
+
+UCLASS()
+class UMockAI_BT : public UMockAI
+{
+	GENERATED_UCLASS_BODY()
+
+	UPROPERTY()
+	UBehaviorTreeComponent* BTComp;
+
+	static TArray<int32> ExecutionLog;
+	TArray<int32> ExpectedResult;
+
+	bool IsRunning() const;
+	bool RunBT(UBehaviorTree* BTAsset, EBTExecutionMode::Type RunType = EBTExecutionMode::SingleRun);
+};
diff --git a/Engine/Source/Developer/AITestSuite/Private/AITestSuite.cpp b/Engine/Source/Developer/AITestSuite/Private/AITestSuite.cpp
--- a/Engine/Source/Developer/AITestSuite/Private/AITestSuite.cpp
+++ b/Engine/Source/Developer/AITestSuite/Private/AITestSuite.cpp
@@ -2,6 +2,8 @@
 
 #include "AITestSuitePrivatePCH.h"
 
+DEFINE_LOG_CATEGORY(LogAITestSuite);
+DEFINE_LOG_CATEGORY(LogBehaviorTreeTest);
 
 class FAITestSuite : public IAITestSuite
 {
diff --git a/Engine/Source/Developer/AITestSuite/Private/AITestSuitePrivatePCH.h b/Engine/Source/Developer/AITestSuite/Private/AITestSuitePrivatePCH.h
--- a/Engine/Source/Developer/AITestSuite/Private/AITestSuitePrivatePCH.h
+++ b/Engine/Source/Developer/AITestSuite/Private/AITestSuitePrivatePCH.h
@@ -2,5 +2,17 @@
 
 #include "AITestSuite.h"
 
-// You should place include statements to your module's private header files here.  You only need to
-// add includes for headers that are used in most of your module's source files though.
+#include "BehaviorTree/Blackboard/BlackboardKeyAllTypes.h"
+#include "BehaviorTree/BlackboardData.h"
+
+#include "BehaviorTree/BlackboardComponent.h"
+#include "BehaviorTree/BehaviorTreeComponent.h"
+#include "BehaviorTree/BehaviorTree.h"
+#include "Perception/AIPerceptionComponent.h"
+#include "Actions/PawnActionsComponent.h"
+
+#include "BTBuilder.h"
+
+#include "Tests/AITestsCommon.h"
+
+#include "MockAI.h"
diff --git a/Engine/Source/Developer/AITestSuite/Private/BehaviorTree/TestBTDecorator_CantExecute.cpp b/Engine/Source/Developer/AITestSuite/Private/BehaviorTree/TestBTDecorator_CantExecute.cpp
new file mode 100644
--- /dev/null
+++ b/Engine/Source/Developer/AITestSuite/Private/BehaviorTree/TestBTDecorator_CantExecute.cpp
@@ -0,0 +1,18 @@
+// Copyright 1998-2014 Epic Games, Inc. All Rights Reserved.
+
+#include "AITestSuitePrivatePCH.h"
+#include "BehaviorTree/TestBTDecorator_CantExecute.h"
+
+UTestBTDecorator_CantExecute::UTestBTDecorator_CantExecute(const class FPostConstructInitializeProperties& PCIP) : Super(PCIP)
+{
+	NodeName = TEXT("Can't Exexcute");
+
+	bAllowAbortNone = false;
+	bAllowAbortLowerPri = false;
+	bAllowAbortChildNodes = false;
+}
+
+bool UTestBTDecorator_CantExecute::CalculateRawConditionValue(class UBehaviorTreeComponent* OwnerComp, uint8* NodeMemory) const
+{
+	return false;
+}
diff --git a/Engine/Source/Developer/AITestSuite/Private/BehaviorTree/TestBTTask_Log.cpp b/Engine/Source/Developer/AITestSuite/Private/BehaviorTree/TestBTTask_Log.cpp
new file mode 100644
--- /dev/null
+++ b/Engine/Source/Developer/AITestSuite/Private/BehaviorTree/TestBTTask_Log.cpp
@@ -0,0 +1,52 @@
+// Copyright 1998-2014 Epic Games, Inc. All Rights Reserved.
+
+#include "AITestSuitePrivatePCH.h"
+#include "BehaviorTree/TestBTTask_Log.h"
+#include "MockAI_BT.h"
+
+UTestBTTask_Log::UTestBTTask_Log(const class FPostConstructInitializeProperties& PCIP) : Super(PCIP)
+{
+	NodeName = "Log";
+	ExecutionTime = 5.0f;
+	LogIndex = 0;
+	LogResult = EBTNodeResult::Succeeded;
+
+	bNotifyTick = true;
+}
+
+EBTNodeResult::Type UTestBTTask_Log::ExecuteTask(class UBehaviorTreeComponent* OwnerComp, uint8* NodeMemory)
+{
+	FBTLogTaskMemory* MyMemory = (FBTLogTaskMemory*)NodeMemory;
+	MyMemory->RemainingWaitTime = ExecutionTime;
+
+	if (ExecutionTime <= 0.0f)
+	{
+		LogExecution(OwnerComp);
+		return LogResult;
+	}
+
+	return EBTNodeResult::InProgress;
+}
+
+void UTestBTTask_Log::TickTask(class UBehaviorTreeComponent* OwnerComp, uint8* NodeMemory, float DeltaSeconds)
+{
+	FBTLogTaskMemory* MyMemory = (FBTLogTaskMemory*)NodeMemory;
+	MyMemory->RemainingWaitTime -= DeltaSeconds;
+
+	if (MyMemory->RemainingWaitTime <= 0.0f)
+	{
+		// continue execution from this node
+		LogExecution(OwnerComp);
+		FinishLatentTask(OwnerComp, LogResult);
+	}
+}
+
+uint16 UTestBTTask_Log::GetInstanceMemorySize() const
+{
+	return sizeof(FBTLogTaskMemory);
+}
+
+void UTestBTTask_Log::LogExecution(class UBehaviorTreeComponent* OwnerComp)
+{
+	UMockAI_BT::ExecutionLog.Add(LogIndex);
+}
diff --git a/Engine/Source/Developer/AITestSuite/Private/BehaviorTree/TestBTTask_SetFlag.cpp b/Engine/Source/Developer/AITestSuite/Private/BehaviorTree/TestBTTask_SetFlag.cpp
new file mode 100644
--- /dev/null
+++ b/Engine/Source/Developer/AITestSuite/Private/BehaviorTree/TestBTTask_SetFlag.cpp
@@ -0,0 +1,18 @@
+// Copyright 1998-2014 Epic Games, Inc. All Rights Reserved.
+
+#include "AITestSuitePrivatePCH.h"
+#include "BehaviorTree/TestBTTask_SetFlag.h"
+#include "BehaviorTree/BlackboardComponent.h"
+
+UTestBTTask_SetFlag::UTestBTTask_SetFlag(const class FPostConstructInitializeProperties& PCIP) : Super(PCIP)
+{
+	NodeName = "Log";
+	TaskResult = EBTNodeResult::Succeeded;
+	bValue = true;
+}
+
+EBTNodeResult::Type UTestBTTask_SetFlag::ExecuteTask(class UBehaviorTreeComponent* OwnerComp, uint8* NodeMemory)
+{
+	OwnerComp->GetBlackboardComponent()->SetValueAsBool(TEXT("Bool1"), bValue);
+	return TaskResult;
+}
diff --git a/Engine/Source/Developer/AITestSuite/Private/MockAI/MockAI.cpp b/Engine/Source/Developer/AITestSuite/Private/MockAI/MockAI.cpp
new file mode 100644
--- /dev/null
+++ b/Engine/Source/Developer/AITestSuite/Private/MockAI/MockAI.cpp
@@ -0,0 +1,84 @@
+// Copyright 1998-2014 Epic Games, Inc. All Rights Reserved.
+
+#include "AITestSuitePrivatePCH.h"
+#include "MockAI.h"
+
+//----------------------------------------------------------------------//
+// 
+//----------------------------------------------------------------------//
+void FTestTickHelper::Tick(float DeltaTime)
+{
+	if (Owner.IsValid())
+	{
+		Owner->TickMe(DeltaTime);
+	}
+}
+
+TStatId FTestTickHelper::GetStatId() const
+{
+	RETURN_QUICK_DECLARE_CYCLE_STAT(FTestTickHelper, STATGROUP_Tickables);
+}
+
+//----------------------------------------------------------------------//
+// 
+//----------------------------------------------------------------------//
+UMockAI::UMockAI(const FPostConstructInitializeProperties& PCIP)
+	: Super(PCIP)
+{
+}
+
+UMockAI::~UMockAI()
+{
+	TickHelper.Owner.Reset();
+}
+
+void UMockAI::SetEnableTicking(bool bShouldTick)
+{
+	if (bShouldTick)
+	{
+		TickHelper.Owner = this;
+	}
+	else
+	{
+		TickHelper.Owner = NULL;
+	}
+}
+
+void UMockAI::UseBlackboardComponent()
+{
+	BBComp = NewObject<UBlackboardComponent>(FAITestHelpers::GetWorld());
+}
+
+void UMockAI::UsePerceptionComponent()
+{
+	PerceptionComp = NewObject<UAIPerceptionComponent>(FAITestHelpers::GetWorld());
+}
+
+void UMockAI::UsePawnActionsComponent()
+{
+	PawnActionComp = NewObject<UPawnActionsComponent>(FAITestHelpers::GetWorld());
+}
+
+void UMockAI::TickMe(float DeltaTime)
+{
+	if (BBComp)
+	{
+		BBComp->TickComponent(DeltaTime, ELevelTick::LEVELTICK_All, NULL);
+	}
+
+	if (PerceptionComp)
+	{
+		PerceptionComp->TickComponent(DeltaTime, ELevelTick::LEVELTICK_All, NULL);
+	}
+	
+	if (BrainComp)
+	{
+		BrainComp->TickComponent(DeltaTime, ELevelTick::LEVELTICK_All, NULL);
+	}
+	
+	if (PawnActionComp)
+	{
+		PawnActionComp->TickComponent(DeltaTime, ELevelTick::LEVELTICK_All, NULL);
+	}
+}
+
diff --git a/Engine/Source/Developer/AITestSuite/Private/MockAI/MockAI_BT.cpp b/Engine/Source/Developer/AITestSuite/Private/MockAI/MockAI_BT.cpp
new file mode 100644
--- /dev/null
+++ b/Engine/Source/Developer/AITestSuite/Private/MockAI/MockAI_BT.cpp
@@ -0,0 +1,39 @@
+// Copyright 1998-2014 Epic Games, Inc. All Rights Reserved.
+
+#include "AITestSuitePrivatePCH.h"
+#include "MockAI_BT.h"
+
+//----------------------------------------------------------------------//
+// 
+//----------------------------------------------------------------------//
+TArray<int32> UMockAI_BT::ExecutionLog;
+
+UMockAI_BT::UMockAI_BT(const FPostConstructInitializeProperties& PCIP)
+	: Super(PCIP)
+{
+	if (HasAnyFlags(RF_ClassDefaultObject) == false)
+	{
+		UseBlackboardComponent();
+		UseBrainComponent<UBehaviorTreeComponent>();
+
+		BTComp = Cast<UBehaviorTreeComponent>(BrainComp);
+	}
+}
+
+bool UMockAI_BT::IsRunning() const
+{
+	return BTComp && BTComp->IsRunning() && BTComp->GetRootTree();
+}
+
+bool UMockAI_BT::RunBT(UBehaviorTree* BTAsset, EBTExecutionMode::Type RunType)
+{
+	BBComp->InitializeBlackboard(BTAsset->BlackboardAsset);
+	BTComp->CacheBlackboardComponent(BBComp);
+
+	UWorld* World = FAITestHelpers::GetWorld();
+
+	BBComp->RegisterComponentWithWorld(World);
+	BTComp->RegisterComponentWithWorld(World);
+
+	return BTComp->StartTree(BTAsset, RunType);
+}
\ No newline at end of file
diff --git a/Engine/Source/Developer/AITestSuite/Private/Tests/AITestsCommon.cpp b/Engine/Source/Developer/AITestSuite/Private/Tests/AITestsCommon.cpp
--- a/Engine/Source/Developer/AITestSuite/Private/Tests/AITestsCommon.cpp
+++ b/Engine/Source/Developer/AITestSuite/Private/Tests/AITestsCommon.cpp
@@ -2,3 +2,121 @@
 
 #include "AITestSuitePrivatePCH.h"
 
+bool FAITestCommand_WaitSeconds::Update()
+{
+	float NewTime = FPlatformTime::Seconds();
+	if (NewTime - StartTime >= Duration)
+	{
+		return true;
+	}
+	return false;
+}
+
+bool FAITestCommand_WaitOneTick::Update()
+{
+	if (bAlreadyRun == false)
+	{
+		bAlreadyRun = true;
+		return true;
+	}
+	return false;
+}
+
+bool FAITestCommand_SetUpTest::Update()
+{
+	if (AITest)
+	{
+		AITest->SetUp();
+	}
+	return true;
+}
+
+bool FAITestCommand_PerformTest::Update()
+{
+	return AITest != nullptr && AITest->Update();
+}
+
+bool FAITestCommand_TearDownTest::Update()
+{
+	if (AITest)
+	{
+		AITest->TearDown();
+		delete AITest;
+		AITest = nullptr;
+	}
+	return true;
+}
+
+namespace FAITestHelpers
+{
+	UWorld* GetWorld()
+	{
+#if WITH_EDITOR
+		if (GIsEditor)
+		{
+			return GWorld;
+		}
+#endif // WITH_EDITOR
+		return GEngine->GetWorldContexts()[0].World();
+	}
+}
+
+//----------------------------------------------------------------------//
+// FAITest_SimpleBT
+//----------------------------------------------------------------------//
+FAITest_SimpleBT::FAITest_SimpleBT()
+{
+	AIBTUser = NewObject<UMockAI_BT>();
+	AIBTUser->AddToRoot();
+	BTAsset = &FBTBuilder::CreateBehaviorTree();
+
+	UMockAI_BT::ExecutionLog.Reset();
+}
+
+FAITest_SimpleBT::~FAITest_SimpleBT()
+{
+	if (AIBTUser)
+	{
+		AIBTUser->RemoveFromRoot();
+		AIBTUser = nullptr;
+	}
+}
+
+void FAITest_SimpleBT::SetUp()
+{
+	if (AIBTUser)
+	{
+		AIBTUser->RunBT(BTAsset, EBTExecutionMode::SingleRun);
+		AIBTUser->SetEnableTicking(true);
+	}
+}
+
+bool FAITest_SimpleBT::Update()
+{
+	if (AIBTUser != NULL && AIBTUser->IsRunning())
+	{
+		return false;
+	}
+
+	VerifyResults();
+	return true;
+}
+
+void FAITest_SimpleBT::VerifyResults()
+{
+	const bool bMatch = (ExpectedResult == UMockAI_BT::ExecutionLog);
+	if (!bMatch)
+	{
+		FString Desc;
+		for (int32 Idx = 0; Idx < UMockAI_BT::ExecutionLog.Num(); Idx++)
+		{
+			Desc += TTypeToString<int32>::ToString(UMockAI_BT::ExecutionLog[Idx]);
+			if (Idx < (UMockAI_BT::ExecutionLog.Num() - 1))
+			{
+				Desc += TEXT(", ");
+			}
+		}
+
+		UE_LOG(LogBehaviorTreeTest, Error, TEXT("Results are not matching!\nExecution log: %s"), *Desc);
+	}
+}
\ No newline at end of file
diff --git a/Engine/Source/Developer/AITestSuite/Private/Tests/AITestsCommon.h b/Engine/Source/Developer/AITestSuite/Private/Tests/AITestsCommon.h
--- a/Engine/Source/Developer/AITestSuite/Private/Tests/AITestsCommon.h
+++ b/Engine/Source/Developer/AITestSuite/Private/Tests/AITestsCommon.h
@@ -2,3 +2,74 @@
 
 #pragma once
 
+#include "Misc/AutomationTest.h"
+
+class UBehaviorTree;
+class UMockAI_BT;
+
+DECLARE_LOG_CATEGORY_EXTERN(LogAITestSuite, Log, All);
+DECLARE_LOG_CATEGORY_EXTERN(LogBehaviorTreeTest, Log, All);
+
+DEFINE_LATENT_AUTOMATION_COMMAND_ONE_PARAMETER(FAITestCommand_WaitSeconds, float, Duration);
+
+class FAITestCommand_WaitOneTick : public IAutomationLatentCommand
+{
+public: 
+	FAITestCommand_WaitOneTick()
+		: bAlreadyRun(false)
+	{} 
+	virtual bool Update() override;
+private: 
+	bool bAlreadyRun;
+};
+
+
+namespace FAITestHelpers
+{
+	UWorld* GetWorld();
+}
+
+struct FAITestBase
+{
+	virtual ~FAITestBase(){}
+	virtual void SetUp() {}
+	virtual bool Update() { return true; }
+	virtual void TearDown() {}
+};
+
+DEFINE_LATENT_AUTOMATION_COMMAND_ONE_PARAMETER(FAITestCommand_SetUpTest, FAITestBase*, AITest);
+DEFINE_LATENT_AUTOMATION_COMMAND_ONE_PARAMETER(FAITestCommand_PerformTest, FAITestBase*, AITest);
+DEFINE_LATENT_AUTOMATION_COMMAND_ONE_PARAMETER(FAITestCommand_TearDownTest, FAITestBase*, AITest);
+
+#define IMPLEMENT_AI_TEST(TestClass, PrettyName) \
+	IMPLEMENT_SIMPLE_AUTOMATION_TEST(TestClass##Runner, PrettyName, (EAutomationTestFlags::ATF_Game | EAutomationTestFlags::ATF_Editor)) \
+	bool TestClass##Runner::RunTest(const FString& Parameters) \
+	{ \
+		/* spawn test instance. Setup should be done in test's constructor */ \
+		TestClass* TestInstance = new TestClass(); \
+		/* set up */ \
+		ADD_LATENT_AUTOMATION_COMMAND(FAITestCommand_SetUpTest(TestInstance)); \
+		/* run latent command to update */ \
+		ADD_LATENT_AUTOMATION_COMMAND(FAITestCommand_PerformTest(TestInstance)); \
+		/* run latent command to tear down */ \
+		ADD_LATENT_AUTOMATION_COMMAND(FAITestCommand_TearDownTest(TestInstance)); \
+		return true; \
+	} 
+
+//----------------------------------------------------------------------//
+// Specific test types
+//----------------------------------------------------------------------//
+
+struct FAITest_SimpleBT : public FAITestBase
+{
+	TArray<int32> ExpectedResult;
+	UBehaviorTree* BTAsset;
+	UMockAI_BT* AIBTUser;
+
+	FAITest_SimpleBT();	
+	virtual ~FAITest_SimpleBT();
+	virtual void SetUp() override;
+	virtual bool Update() override;
+	
+	virtual void VerifyResults();
+};
\ No newline at end of file
diff --git a/Engine/Source/Developer/AITestSuite/Private/Tests/BTTest.cpp b/Engine/Source/Developer/AITestSuite/Private/Tests/BTTest.cpp
--- a/Engine/Source/Developer/AITestSuite/Private/Tests/BTTest.cpp
+++ b/Engine/Source/Developer/AITestSuite/Private/Tests/BTTest.cpp
@@ -1,30 +1,64 @@
 // Copyright 1998-2014 Epic Games, Inc. All Rights Reserved.
 
 #include "AITestSuitePrivatePCH.h"
+#include "MockAI_BT.h"
+#include "BehaviorTree/TestBTDecorator_CantExecute.h"
 
 #define LOCTEXT_NAMESPACE "AITestSuite_BTTest"
 
-// create test base class
-IMPLEMENT_COMPLEX_AUTOMATION_TEST(FAITestSuite_BehaviorTree, "Engine.AI.Behavior Trees", (EAutomationTestFlags::ATF_Game | EAutomationTestFlags::ATF_Editor))
+//----------------------------------------------------------------------//
+// 
+//----------------------------------------------------------------------//
+struct FAITest_BasicSelector : public FAITest_SimpleBT
+{
+	FAITest_BasicSelector()
+	{
+		UBTCompositeNode& CompNode = FBTBuilder::AddSelector(*BTAsset);
+		{
+			FBTBuilder::AddTask(CompNode, 0, EBTNodeResult::Failed);
+
+			FBTBuilder::AddTask(CompNode, 1, EBTNodeResult::Succeeded);
+			{
+				FBTBuilder::WithDecorator<UTestBTDecorator_CantExecute>(CompNode);
+			}
+
+			FBTBuilder::AddTask(CompNode, 2, EBTNodeResult::Succeeded, 0.5f);
+
+			FBTBuilder::AddTask(CompNode, 3, EBTNodeResult::Failed);
+		}
+
+		ExpectedResult.Add(0);
+		ExpectedResult.Add(2);
+	}
+};
+IMPLEMENT_AI_TEST(FAITest_BasicSelector, "Engine.AI.Behavior Trees.Composite node: selector")
 
-/** 
- * Requests a enumeration of all maps to be loaded
- */
-void FAITestSuite_BehaviorTree::GetTests(TArray<FString>& OutBeautifiedNames, TArray<FString>& OutTestCommands) const
+//----------------------------------------------------------------------//
+// 
+//----------------------------------------------------------------------//
+struct FAITest_BasicSequence : public FAITest_SimpleBT
 {
-	OutBeautifiedNames.Add(TEXT("Simple Parallel aborting"));
-	OutTestCommands.Add(TEXT("BTTestID"));
-}
-
-/** 
- * Execute the loading of each map and performance captures
- *
- * @param Parameters - Should specify which map name to load
- * @return	TRUE if the test was successful, FALSE otherwise
- */
-bool FAITestSuite_BehaviorTree::RunTest(const FString& Parameters)
-{	
-	return true;
-}
-
-#undef LOCTEXT_NAMESPACE
\ No newline at end of file
+	FAITest_BasicSequence()
+	{
+		UBTCompositeNode& CompNode = FBTBuilder::AddSequence(*BTAsset);
+		{
+			FBTBuilder::AddTask(CompNode, 0, EBTNodeResult::Succeeded);
+
+			FBTBuilder::AddTask(CompNode, 1, EBTNodeResult::Failed);
+			{
+				FBTBuilder::WithDecorator<UTestBTDecorator_CantExecute>(CompNode);
+				FBTBuilder::WithDecorator<UBTDecorator_ForceSuccess>(CompNode);
+			}
+
+			FBTBuilder::AddTask(CompNode, 2, EBTNodeResult::Failed, 0.5f);
+
+			FBTBuilder::AddTask(CompNode, 3, EBTNodeResult::Succeeded);
+		}
+
+		ExpectedResult.Add(0);
+		ExpectedResult.Add(2);
+	}
+};
+IMPLEMENT_AI_TEST(FAITest_BasicSequence, "Engine.AI.Behavior Trees.Composite node: sequence")
+
+#undef LOCTEXT_NAMESPACE
diff --git a/Engine/Source/Developer/AITestSuite/Public/BTBuilder.h b/Engine/Source/Developer/AITestSuite/Public/BTBuilder.h
new file mode 100644
--- /dev/null
+++ b/Engine/Source/Developer/AITestSuite/Public/BTBuilder.h
@@ -0,0 +1,98 @@
+// Copyright 1998-2014 Epic Games, Inc. All Rights Reserved.
+
+#pragma once
+
+#include "BehaviorTree/BehaviorTree.h"
+#include "BehaviorTree/Composites/BTComposite_Selector.h"
+#include "BehaviorTree/Composites/BTComposite_Sequence.h"
+#include "BehaviorTree/Composites/BTComposite_SimpleParallel.h"
+#include "BehaviorTree/Decorators//BTDecorator_ForceSuccess.h"
+
+#include "BehaviorTree/TestBTTask_Log.h"
+
+struct FBTBuilder
+{
+	static UBehaviorTree& CreateBehaviorTree()
+	{
+		UBlackboardData* BB = NewObject<UBlackboardData>();
+		FBlackboardEntry KeyData;
+
+		KeyData.EntryName = TEXT("Bool1");
+		KeyData.KeyType = NewObject<UBlackboardKeyType_Bool>();
+		BB->Keys.Add(KeyData);
+
+		KeyData.EntryName = TEXT("Bool2");
+		KeyData.KeyType = NewObject<UBlackboardKeyType_Bool>();
+		BB->Keys.Add(KeyData);
+
+		UBehaviorTree* TreeOb = NewObject<UBehaviorTree>();
+		TreeOb->BlackboardAsset = BB;
+
+		return *TreeOb;
+	}
+
+	static UBTComposite_Selector& AddSelector(UBehaviorTree& TreeOb)
+	{
+		UBTComposite_Selector* NodeOb = NewObject<UBTComposite_Selector>(&TreeOb);
+		NodeOb->InitializeFromAsset(&TreeOb);
+		TreeOb.RootNode = NodeOb;
+		return *NodeOb;
+	}
+
+	static UBTComposite_Selector& AddSelector(UBTCompositeNode& ParentNode)
+	{
+		UBTComposite_Selector* NodeOb = NewObject<UBTComposite_Selector>(ParentNode.GetTreeAsset());
+		NodeOb->InitializeFromAsset(ParentNode.GetTreeAsset());
+		return *NodeOb;
+	}
+
+	static UBTComposite_Sequence& AddSequence(UBehaviorTree& TreeOb)
+	{
+		UBTComposite_Sequence* NodeOb = NewObject<UBTComposite_Sequence>(&TreeOb);
+		NodeOb->InitializeFromAsset(&TreeOb);
+		TreeOb.RootNode = NodeOb;
+		return *NodeOb;
+	}
+
+	static UBTComposite_Sequence& AddSequence(UBTCompositeNode& ParentNode)
+	{
+		UBTComposite_Sequence* NodeOb = NewObject<UBTComposite_Sequence>(ParentNode.GetTreeAsset());
+		NodeOb->InitializeFromAsset(ParentNode.GetTreeAsset());
+		return *NodeOb;
+	}
+
+	static UBTComposite_SimpleParallel& AddParallel(UBehaviorTree& TreeOb)
+	{
+		UBTComposite_SimpleParallel* NodeOb = NewObject<UBTComposite_SimpleParallel>(&TreeOb);
+		NodeOb->InitializeFromAsset(&TreeOb);
+		TreeOb.RootNode = NodeOb;
+		return *NodeOb;
+	}
+
+	static UBTComposite_SimpleParallel& AddParallel(UBTCompositeNode& ParentNode)
+	{
+		UBTComposite_SimpleParallel* NodeOb = NewObject<UBTComposite_SimpleParallel>(ParentNode.GetTreeAsset());
+		NodeOb->InitializeFromAsset(ParentNode.GetTreeAsset());
+		return *NodeOb;
+	}
+
+	static void AddTask(UBTCompositeNode& ParentNode, int32 LogIndex, EBTNodeResult::Type NodeResult, float ExecutionTime = 0.0f)
+	{
+		UTestBTTask_Log* TaskNode = NewObject<UTestBTTask_Log>(ParentNode.GetTreeAsset());
+		TaskNode->LogIndex = LogIndex;
+		TaskNode->LogResult = NodeResult;
+		TaskNode->ExecutionTime = ExecutionTime;
+
+		const int32 ChildIdx = ParentNode.Children.AddZeroed(1);
+		ParentNode.Children[ChildIdx].ChildTask = TaskNode;
+	}
+
+	template<class T>
+	static T& WithDecorator(UBTCompositeNode& ParentNode, UClass* DecoratorClass = T::StaticClass())
+	{
+		T* DecoratorOb = NewObject<T>(ParentNode.GetTreeAsset());
+		ParentNode.Children.Last().Decorators.Add(DecoratorOb);
+
+		return *DecoratorOb;
+	}
+};
diff --git a/Engine/Source/Runtime/AIModule/Classes/Perception/AIPerceptionTypes.h b/Engine/Source/Runtime/AIModule/Classes/Perception/AIPerceptionTypes.h
--- a/Engine/Source/Runtime/AIModule/Classes/Perception/AIPerceptionTypes.h
+++ b/Engine/Source/Runtime/AIModule/Classes/Perception/AIPerceptionTypes.h
@@ -2,6 +2,7 @@
 
 #pragma once
 
+#include "GenericTeamAgentInterface.h"
 #include "AIPerceptionTypes.generated.h"
 
 class UAIPerceptionComponent;
