AITestSuite non-unity compilation fix #UE4

diff --git a/Engine/Source/Developer/AITestSuite/Private/AITestSuitePrivatePCH.h b/Engine/Source/Developer/AITestSuite/Private/AITestSuitePrivatePCH.h
--- a/Engine/Source/Developer/AITestSuite/Private/AITestSuitePrivatePCH.h
+++ b/Engine/Source/Developer/AITestSuite/Private/AITestSuitePrivatePCH.h
@@ -16,3 +16,4 @@
 #include "Tests/AITestsCommon.h"
 
 #include "MockAI.h"
+#include "MockAI_BT.h"
