disabled automation tests for parallel BT nodes, it can't be time based

diff --git a/Engine/Source/Developer/AITestSuite/Private/Tests/BTTest.cpp b/Engine/Source/Developer/AITestSuite/Private/Tests/BTTest.cpp
--- a/Engine/Source/Developer/AITestSuite/Private/Tests/BTTest.cpp
+++ b/Engine/Source/Developer/AITestSuite/Private/Tests/BTTest.cpp
@@ -88,7 +88,7 @@ struct FAITest_BTBasicParallelWait : public FAITest_SimpleBT
 		ExpectedResult.Add(3);
 	}
 };
-IMPLEMENT_AI_TEST(FAITest_BTBasicParallelWait, "Engine.AI.Behavior Trees.Composite node: parallel (wait)")
+//IMPLEMENT_AI_TEST(FAITest_BTBasicParallelWait, "Engine.AI.Behavior Trees.Composite node: parallel (wait)")
 
 struct FAITest_BTBasicParallelAbort : public FAITest_SimpleBT
 {
@@ -120,7 +120,7 @@ struct FAITest_BTBasicParallelAbort : public FAITest_SimpleBT
 		ExpectedResult.Add(3);
 	}
 };
-IMPLEMENT_AI_TEST(FAITest_BTBasicParallelAbort, "Engine.AI.Behavior Trees.Composite node: parallel (abort)")
+//IMPLEMENT_AI_TEST(FAITest_BTBasicParallelAbort, "Engine.AI.Behavior Trees.Composite node: parallel (abort)")
 
 struct FAITest_BTCompositeDecorator : public FAITest_SimpleBT
 {
@@ -373,7 +373,7 @@ struct FAITest_BTAbortParallelInternal : public FAITest_SimpleBT
 		ExpectedResult.Add(4);
 	}
 };
-IMPLEMENT_AI_TEST(FAITest_BTAbortParallelInternal, "Engine.AI.Behavior Trees.Abort: parallel internal")
+//IMPLEMENT_AI_TEST(FAITest_BTAbortParallelInternal, "Engine.AI.Behavior Trees.Abort: parallel internal")
 
 struct FAITest_BTAbortParallelOut : public FAITest_SimpleBT
 {
@@ -406,7 +406,7 @@ struct FAITest_BTAbortParallelOut : public FAITest_SimpleBT
 		ExpectedResult.Add(0);
 	}
 };
-IMPLEMENT_AI_TEST(FAITest_BTAbortParallelOut, "Engine.AI.Behavior Trees.Abort: parallel out")
+//IMPLEMENT_AI_TEST(FAITest_BTAbortParallelOut, "Engine.AI.Behavior Trees.Abort: parallel out")
 
 struct FAITest_BTAbortParallelOutAndBack : public FAITest_SimpleBT
 {
@@ -441,7 +441,7 @@ struct FAITest_BTAbortParallelOutAndBack : public FAITest_SimpleBT
 		ExpectedResult.Add(4);
 	}
 };
-IMPLEMENT_AI_TEST(FAITest_BTAbortParallelOutAndBack, "Engine.AI.Behavior Trees.Abort: parallel out & back")
+//IMPLEMENT_AI_TEST(FAITest_BTAbortParallelOutAndBack, "Engine.AI.Behavior Trees.Abort: parallel out & back")
 
 struct FAITest_BTSubtreeSimple : public FAITest_SimpleBT
 {
