changed automated BT tests to use frame based delay

diff --git a/Engine/Source/Developer/AITestSuite/Classes/BehaviorTree/TestBTTask_Log.h b/Engine/Source/Developer/AITestSuite/Classes/BehaviorTree/TestBTTask_Log.h
--- a/Engine/Source/Developer/AITestSuite/Classes/BehaviorTree/TestBTTask_Log.h
+++ b/Engine/Source/Developer/AITestSuite/Classes/BehaviorTree/TestBTTask_Log.h
@@ -6,8 +6,7 @@
 
 struct FBTLogTaskMemory
 {
-	/** time left */
-	float RemainingWaitTime;
+	uint64 EndFrameIdx;
 };
 
 UCLASS(meta=(HiddenNode))
@@ -22,7 +21,7 @@ class UTestBTTask_Log : public UBTTaskNode
 	int32 LogFinished;
 
 	UPROPERTY()
-	float ExecutionTime;
+	int32 ExecutionTicks;
 
 	UPROPERTY()
 	TEnumAsByte<EBTNodeResult::Type> LogResult;
diff --git a/Engine/Source/Developer/AITestSuite/Private/BehaviorTree/TestBTTask_Log.cpp b/Engine/Source/Developer/AITestSuite/Private/BehaviorTree/TestBTTask_Log.cpp
--- a/Engine/Source/Developer/AITestSuite/Private/BehaviorTree/TestBTTask_Log.cpp
+++ b/Engine/Source/Developer/AITestSuite/Private/BehaviorTree/TestBTTask_Log.cpp
@@ -7,7 +7,7 @@
 UTestBTTask_Log::UTestBTTask_Log(const FPostConstructInitializeProperties& PCIP) : Super(PCIP)
 {
 	NodeName = "Log";
-	ExecutionTime = 5.0f;
+	ExecutionTicks = 0;
 	LogIndex = 0;
 	LogFinished = -1;
 	LogResult = EBTNodeResult::Succeeded;
@@ -18,10 +18,10 @@ UTestBTTask_Log::UTestBTTask_Log(const FPostConstructInitializeProperties& PCIP)
 EBTNodeResult::Type UTestBTTask_Log::ExecuteTask(UBehaviorTreeComponent* OwnerComp, uint8* NodeMemory)
 {
 	FBTLogTaskMemory* MyMemory = (FBTLogTaskMemory*)NodeMemory;
-	MyMemory->RemainingWaitTime = ExecutionTime;
+	MyMemory->EndFrameIdx = ExecutionTicks + GFrameCounter;
 
 	LogExecution(OwnerComp, LogIndex);
-	if (ExecutionTime <= 0.0f)
+	if (ExecutionTicks == 0)
 	{
 		return LogResult;
 	}
@@ -32,9 +32,8 @@ EBTNodeResult::Type UTestBTTask_Log::ExecuteTask(UBehaviorTreeComponent* OwnerCo
 void UTestBTTask_Log::TickTask(UBehaviorTreeComponent* OwnerComp, uint8* NodeMemory, float DeltaSeconds)
 {
 	FBTLogTaskMemory* MyMemory = (FBTLogTaskMemory*)NodeMemory;
-	MyMemory->RemainingWaitTime -= DeltaSeconds;
 
-	if (MyMemory->RemainingWaitTime <= 0.0f)
+	if (GFrameCounter >= MyMemory->EndFrameIdx)
 	{
 		LogExecution(OwnerComp, LogFinished);
 		FinishLatentTask(OwnerComp, LogResult);
diff --git a/Engine/Source/Developer/AITestSuite/Private/Tests/BTTest.cpp b/Engine/Source/Developer/AITestSuite/Private/Tests/BTTest.cpp
--- a/Engine/Source/Developer/AITestSuite/Private/Tests/BTTest.cpp
+++ b/Engine/Source/Developer/AITestSuite/Private/Tests/BTTest.cpp
@@ -19,7 +19,7 @@ struct FAITest_BTBasicSelector : public FAITest_SimpleBT
 				FBTBuilder::WithDecorator<UTestBTDecorator_CantExecute>(CompNode);
 			}
 
-			FBTBuilder::AddTask(CompNode, 2, EBTNodeResult::Succeeded, 0.5f);
+			FBTBuilder::AddTask(CompNode, 2, EBTNodeResult::Succeeded, 2);
 
 			FBTBuilder::AddTask(CompNode, 3, EBTNodeResult::Failed);
 		}
@@ -44,7 +44,7 @@ struct FAITest_BTBasicSequence : public FAITest_SimpleBT
 				FBTBuilder::WithDecorator<UBTDecorator_ForceSuccess>(CompNode);
 			}
 
-			FBTBuilder::AddTask(CompNode, 2, EBTNodeResult::Failed, 0.5f);
+			FBTBuilder::AddTask(CompNode, 2, EBTNodeResult::Failed, 2);
 
 			FBTBuilder::AddTask(CompNode, 3, EBTNodeResult::Succeeded);
 		}
@@ -63,32 +63,31 @@ struct FAITest_BTBasicParallelWait : public FAITest_SimpleBT
 		{
 			UBTCompositeNode& CompNode2 = FBTBuilder::AddParallel(CompNode, EBTParallelMode::WaitForBackground);
 			{
-				FBTBuilder::AddTaskLogFinish(CompNode2, 0, 10, EBTNodeResult::Succeeded, 1.0f);
+				FBTBuilder::AddTaskLogFinish(CompNode2, 0, 10, EBTNodeResult::Succeeded, 6);
 
 				UBTCompositeNode& CompNode3 = FBTBuilder::AddSequence(CompNode2);
 				{
-					FBTBuilder::AddTaskLogFinish(CompNode3, 1, 11, EBTNodeResult::Succeeded, 0.4f);
-					FBTBuilder::AddTaskLogFinish(CompNode3, 2, 12, EBTNodeResult::Succeeded, 0.4f);
+					FBTBuilder::AddTaskLogFinish(CompNode3, 1, 11, EBTNodeResult::Succeeded, 3);
+					FBTBuilder::AddTaskLogFinish(CompNode3, 2, 12, EBTNodeResult::Succeeded, 3);
+					FBTBuilder::AddTaskLogFinish(CompNode3, 3, 13, EBTNodeResult::Succeeded, 3);
 				}
 			}
 
-			FBTBuilder::AddTask(CompNode, 3, EBTNodeResult::Succeeded);
+			FBTBuilder::AddTask(CompNode, 4, EBTNodeResult::Succeeded);
 		}
 
 		ExpectedResult.Add(0);
 			ExpectedResult.Add(1);
 			ExpectedResult.Add(11);
 			ExpectedResult.Add(2);
-			ExpectedResult.Add(12);
-			ExpectedResult.Add(1);
 		ExpectedResult.Add(10);
-			ExpectedResult.Add(11);
-			ExpectedResult.Add(2);
 			ExpectedResult.Add(12);
-		ExpectedResult.Add(3);
+			ExpectedResult.Add(3);
+			ExpectedResult.Add(13);
+		ExpectedResult.Add(4);
 	}
 };
-//IMPLEMENT_AI_TEST(FAITest_BTBasicParallelWait, "Engine.AI.Behavior Trees.Composite node: parallel (wait)")
+IMPLEMENT_AI_TEST(FAITest_BTBasicParallelWait, "Engine.AI.Behavior Trees.Composite node: parallel (wait)")
 
 struct FAITest_BTBasicParallelAbort : public FAITest_SimpleBT
 {
@@ -98,29 +97,28 @@ struct FAITest_BTBasicParallelAbort : public FAITest_SimpleBT
 		{
 			UBTCompositeNode& CompNode2 = FBTBuilder::AddParallel(CompNode, EBTParallelMode::AbortBackground);
 			{
-				FBTBuilder::AddTaskLogFinish(CompNode2, 0, 10, EBTNodeResult::Succeeded, 1.0f);
+				FBTBuilder::AddTaskLogFinish(CompNode2, 0, 10, EBTNodeResult::Succeeded, 6);
 
 				UBTCompositeNode& CompNode3 = FBTBuilder::AddSequence(CompNode2);
 				{
-					FBTBuilder::AddTaskLogFinish(CompNode3, 1, 11, EBTNodeResult::Succeeded, 0.4f);
-					FBTBuilder::AddTaskLogFinish(CompNode3, 2, 12, EBTNodeResult::Succeeded, 0.4f);
+					FBTBuilder::AddTaskLogFinish(CompNode3, 1, 11, EBTNodeResult::Succeeded, 3);
+					FBTBuilder::AddTaskLogFinish(CompNode3, 2, 12, EBTNodeResult::Succeeded, 3);
+					FBTBuilder::AddTaskLogFinish(CompNode3, 3, 13, EBTNodeResult::Succeeded, 3);
 				}
 			}
 
-			FBTBuilder::AddTask(CompNode, 3, EBTNodeResult::Succeeded);
+			FBTBuilder::AddTask(CompNode, 4, EBTNodeResult::Succeeded);
 		}
 
 		ExpectedResult.Add(0);
 			ExpectedResult.Add(1);
 			ExpectedResult.Add(11);
 			ExpectedResult.Add(2);
-			ExpectedResult.Add(12);
-			ExpectedResult.Add(1);
 		ExpectedResult.Add(10);
-		ExpectedResult.Add(3);
+		ExpectedResult.Add(4);
 	}
 };
-//IMPLEMENT_AI_TEST(FAITest_BTBasicParallelAbort, "Engine.AI.Behavior Trees.Composite node: parallel (abort)")
+IMPLEMENT_AI_TEST(FAITest_BTBasicParallelAbort, "Engine.AI.Behavior Trees.Composite node: parallel (abort)")
 
 struct FAITest_BTCompositeDecorator : public FAITest_SimpleBT
 {
@@ -343,20 +341,20 @@ struct FAITest_BTAbortParallelInternal : public FAITest_SimpleBT
 		{
 			UBTCompositeNode& CompNode2 = FBTBuilder::AddParallel(CompNode, EBTParallelMode::WaitForBackground);
 			{
-				FBTBuilder::AddTask(CompNode2, 0, EBTNodeResult::Succeeded, 0.2f);
+				FBTBuilder::AddTask(CompNode2, 0, EBTNodeResult::Succeeded, 5);
 
 				UBTCompositeNode& CompNode3 = FBTBuilder::AddSequence(CompNode2);
 				{
-					FBTBuilder::AddTask(CompNode3, 1, EBTNodeResult::Succeeded, 0.1f);
+					FBTBuilder::AddTask(CompNode3, 1, EBTNodeResult::Succeeded, 1);
 
 					UBTCompositeNode& CompNode4 = FBTBuilder::AddSelector(CompNode3);
 					{
-						FBTBuilder::AddTask(CompNode4, 2, EBTNodeResult::Succeeded, 0.1f);
+						FBTBuilder::AddTask(CompNode4, 2, EBTNodeResult::Succeeded, 1);
 						{
 							FBTBuilder::WithDecoratorBlackboard(CompNode4, EBasicKeyOperation::Set, EBTFlowAbortMode::LowerPriority);
 						}
 
-						FBTBuilder::AddTask(CompNode4, 3, EBTNodeResult::Succeeded, 0.1f);
+						FBTBuilder::AddTask(CompNode4, 3, EBTNodeResult::Succeeded, 1);
 					}
 
 					FBTBuilder::AddTaskFlagChange(CompNode3, true, EBTNodeResult::Succeeded);
@@ -373,7 +371,7 @@ struct FAITest_BTAbortParallelInternal : public FAITest_SimpleBT
 		ExpectedResult.Add(4);
 	}
 };
-//IMPLEMENT_AI_TEST(FAITest_BTAbortParallelInternal, "Engine.AI.Behavior Trees.Abort: parallel internal")
+IMPLEMENT_AI_TEST(FAITest_BTAbortParallelInternal, "Engine.AI.Behavior Trees.Abort: parallel internal")
 
 struct FAITest_BTAbortParallelOut : public FAITest_SimpleBT
 {
@@ -388,13 +386,13 @@ struct FAITest_BTAbortParallelOut : public FAITest_SimpleBT
 
 			UBTCompositeNode& CompNode2 = FBTBuilder::AddParallel(CompNode, EBTParallelMode::WaitForBackground);
 			{
-				FBTBuilder::AddTask(CompNode2, 1, EBTNodeResult::Failed, 0.2f);
+				FBTBuilder::AddTask(CompNode2, 1, EBTNodeResult::Failed, 5);
 
 				UBTCompositeNode& CompNode3 = FBTBuilder::AddSequence(CompNode2);
 				{
-					FBTBuilder::AddTask(CompNode3, 2, EBTNodeResult::Succeeded, 0.1f);
+					FBTBuilder::AddTask(CompNode3, 2, EBTNodeResult::Succeeded, 1);
 					FBTBuilder::AddTaskFlagChange(CompNode3, true, EBTNodeResult::Succeeded);
-					FBTBuilder::AddTask(CompNode3, 3, EBTNodeResult::Succeeded, 0.1f);
+					FBTBuilder::AddTask(CompNode3, 3, EBTNodeResult::Succeeded, 1);
 				}
 			}
 
@@ -406,7 +404,7 @@ struct FAITest_BTAbortParallelOut : public FAITest_SimpleBT
 		ExpectedResult.Add(0);
 	}
 };
-//IMPLEMENT_AI_TEST(FAITest_BTAbortParallelOut, "Engine.AI.Behavior Trees.Abort: parallel out")
+IMPLEMENT_AI_TEST(FAITest_BTAbortParallelOut, "Engine.AI.Behavior Trees.Abort: parallel out")
 
 struct FAITest_BTAbortParallelOutAndBack : public FAITest_SimpleBT
 {
@@ -422,13 +420,13 @@ struct FAITest_BTAbortParallelOutAndBack : public FAITest_SimpleBT
 
 			UBTCompositeNode& CompNode2 = FBTBuilder::AddParallel(CompNode, EBTParallelMode::WaitForBackground);
 			{
-				FBTBuilder::AddTask(CompNode2, 1, EBTNodeResult::Failed, 0.2f);
+				FBTBuilder::AddTask(CompNode2, 1, EBTNodeResult::Failed, 5);
 
 				UBTCompositeNode& CompNode3 = FBTBuilder::AddSequence(CompNode2);
 				{
-					FBTBuilder::AddTask(CompNode3, 2, EBTNodeResult::Succeeded, 0.1f);
+					FBTBuilder::AddTask(CompNode3, 2, EBTNodeResult::Succeeded, 1);
 					FBTBuilder::AddTaskFlagChange(CompNode3, true, EBTNodeResult::Succeeded);
-					FBTBuilder::AddTask(CompNode3, 3, EBTNodeResult::Succeeded, 0.1f);
+					FBTBuilder::AddTask(CompNode3, 3, EBTNodeResult::Succeeded, 1);
 				}
 			}
 
@@ -441,7 +439,7 @@ struct FAITest_BTAbortParallelOutAndBack : public FAITest_SimpleBT
 		ExpectedResult.Add(4);
 	}
 };
-//IMPLEMENT_AI_TEST(FAITest_BTAbortParallelOutAndBack, "Engine.AI.Behavior Trees.Abort: parallel out & back")
+IMPLEMENT_AI_TEST(FAITest_BTAbortParallelOutAndBack, "Engine.AI.Behavior Trees.Abort: parallel out & back")
 
 struct FAITest_BTSubtreeSimple : public FAITest_SimpleBT
 {
diff --git a/Engine/Source/Developer/AITestSuite/Public/BTBuilder.h b/Engine/Source/Developer/AITestSuite/Public/BTBuilder.h
--- a/Engine/Source/Developer/AITestSuite/Public/BTBuilder.h
+++ b/Engine/Source/Developer/AITestSuite/Public/BTBuilder.h
@@ -109,24 +109,24 @@ struct FBTBuilder
 		return *NodeOb;
 	}
 
-	static void AddTask(UBTCompositeNode& ParentNode, int32 LogIndex, EBTNodeResult::Type NodeResult, float ExecutionTime = 0.0f)
+	static void AddTask(UBTCompositeNode& ParentNode, int32 LogIndex, EBTNodeResult::Type NodeResult, int32 ExecutionTicks = 0)
 	{
 		UTestBTTask_Log* TaskNode = NewObject<UTestBTTask_Log>(ParentNode.GetTreeAsset());
 		TaskNode->LogIndex = LogIndex;
 		TaskNode->LogResult = NodeResult;
-		TaskNode->ExecutionTime = ExecutionTime;
+		TaskNode->ExecutionTicks = ExecutionTicks;
 
 		const int32 ChildIdx = ParentNode.Children.AddZeroed(1);
 		ParentNode.Children[ChildIdx].ChildTask = TaskNode;
 	}
 
-	static void AddTaskLogFinish(UBTCompositeNode& ParentNode, int32 LogIndex, int32 FinishIndex, EBTNodeResult::Type NodeResult, float ExecutionTime = 0.0f)
+	static void AddTaskLogFinish(UBTCompositeNode& ParentNode, int32 LogIndex, int32 FinishIndex, EBTNodeResult::Type NodeResult, int32 ExecutionTicks = 0)
 	{
 		UTestBTTask_Log* TaskNode = NewObject<UTestBTTask_Log>(ParentNode.GetTreeAsset());
 		TaskNode->LogIndex = LogIndex;
 		TaskNode->LogFinished = FinishIndex;
 		TaskNode->LogResult = NodeResult;
-		TaskNode->ExecutionTime = ExecutionTime;
+		TaskNode->ExecutionTicks = ExecutionTicks;
 
 		const int32 ChildIdx = ParentNode.Children.AddZeroed(1);
 		ParentNode.Children[ChildIdx].ChildTask = TaskNode;
