Made the way mock BTTasks count ticks GFramesCounter-independent #UE4

diff --git a/Engine/Source/Developer/AITestSuite/Private/BehaviorTree/TestBTDecorator_DelayedAbort.cpp b/Engine/Source/Developer/AITestSuite/Private/BehaviorTree/TestBTDecorator_DelayedAbort.cpp
--- a/Engine/Source/Developer/AITestSuite/Private/BehaviorTree/TestBTDecorator_DelayedAbort.cpp
+++ b/Engine/Source/Developer/AITestSuite/Private/BehaviorTree/TestBTDecorator_DelayedAbort.cpp
@@ -21,14 +21,14 @@ UTestBTDecorator_DelayedAbort::UTestBTDecorator_DelayedAbort(const FPostConstruc
 void UTestBTDecorator_DelayedAbort::OnBecomeRelevant(UBehaviorTreeComponent* OwnerComp, uint8* NodeMemory)
 {
 	FBTDelayedAbortMemory* MyMemory = (FBTDelayedAbortMemory*)NodeMemory;
-	MyMemory->EndFrameIdx = GFrameCounter + DelayTicks;
+	MyMemory->EndFrameIdx = FAITestHelpers::FramesCounter() + DelayTicks;
 }
 
 void UTestBTDecorator_DelayedAbort::TickNode(UBehaviorTreeComponent* OwnerComp, uint8* NodeMemory, float DeltaSeconds)
 {
 	FBTDelayedAbortMemory* MyMemory = (FBTDelayedAbortMemory*)NodeMemory;
 
-	if (GFrameCounter >= MyMemory->EndFrameIdx)
+	if (FAITestHelpers::FramesCounter() >= MyMemory->EndFrameIdx)
 	{
 		OwnerComp->RequestExecution(this);
 		MyMemory->EndFrameIdx = bOnlyOnce ? MAX_uint64 : 0;
diff --git a/Engine/Source/Developer/AITestSuite/Private/BehaviorTree/TestBTTask_Log.cpp b/Engine/Source/Developer/AITestSuite/Private/BehaviorTree/TestBTTask_Log.cpp
--- a/Engine/Source/Developer/AITestSuite/Private/BehaviorTree/TestBTTask_Log.cpp
+++ b/Engine/Source/Developer/AITestSuite/Private/BehaviorTree/TestBTTask_Log.cpp
@@ -18,7 +18,7 @@ UTestBTTask_Log::UTestBTTask_Log(const FPostConstructInitializeProperties& PCIP)
 EBTNodeResult::Type UTestBTTask_Log::ExecuteTask(UBehaviorTreeComponent* OwnerComp, uint8* NodeMemory)
 {
 	FBTLogTaskMemory* MyMemory = (FBTLogTaskMemory*)NodeMemory;
-	MyMemory->EndFrameIdx = ExecutionTicks + GFrameCounter;
+	MyMemory->EndFrameIdx = ExecutionTicks + FAITestHelpers::FramesCounter();
 
 	LogExecution(OwnerComp, LogIndex);
 	if (ExecutionTicks == 0)
@@ -33,7 +33,7 @@ void UTestBTTask_Log::TickTask(UBehaviorTreeComponent* OwnerComp, uint8* NodeMem
 {
 	FBTLogTaskMemory* MyMemory = (FBTLogTaskMemory*)NodeMemory;
 
-	if (GFrameCounter >= MyMemory->EndFrameIdx)
+	if (FAITestHelpers::FramesCounter() >= MyMemory->EndFrameIdx)
 	{
 		LogExecution(OwnerComp, LogFinished);
 		FinishLatentTask(OwnerComp, LogResult);
diff --git a/Engine/Source/Developer/AITestSuite/Private/Tests/AITestsCommon.cpp b/Engine/Source/Developer/AITestSuite/Private/Tests/AITestsCommon.cpp
--- a/Engine/Source/Developer/AITestSuite/Private/Tests/AITestsCommon.cpp
+++ b/Engine/Source/Developer/AITestSuite/Private/Tests/AITestsCommon.cpp
@@ -2,6 +2,26 @@
 
 #include "AITestSuitePrivatePCH.h"
 
+namespace FAITestHelpers
+{
+	uint64 UpdatesCounter = 0;
+	
+	void UpdateFrameCounter()
+	{
+		static uint64 PreviousFramesCounter = GFrameCounter;
+		if (PreviousFramesCounter != GFrameCounter)
+		{
+			++UpdatesCounter;
+			PreviousFramesCounter = GFrameCounter;
+		}
+	}
+
+	uint64 FramesCounter()
+	{
+		return UpdatesCounter;
+	}
+}
+
 bool FAITestCommand_WaitSeconds::Update()
 {
 	float NewTime = FPlatformTime::Seconds();
@@ -106,6 +126,8 @@ void FAITest_SimpleBT::SetUp()
 
 bool FAITest_SimpleBT::Update()
 {
+	FAITestHelpers::UpdateFrameCounter();
+
 	if (AIBTUser != NULL && AIBTUser->IsRunning())
 	{
 		return false;
@@ -130,7 +152,16 @@ void FAITest_SimpleBT::VerifyResults()
 			}
 		}
 
-		UE_LOG(LogBehaviorTreeTest, Error, TEXT("Results are not matching!\nExecution log: %s"), *Desc);
+		for (int32 Idx = 0; Idx < ExpectedResult.Num(); Idx++)
+		{
+			Desc += TTypeToString<int32>::ToString(ExpectedResult[Idx]);
+			if (Idx < (ExpectedResult.Num() - 1))
+			{
+				Desc += TEXT(", ");
+			}
+		}
+
+		UE_LOG(LogBehaviorTreeTest, Error, TEXT("Test scenario failed to produce expected results!\nExecution log: %s\nExpected values: %s"), *Desc);
 	}
 }
 
diff --git a/Engine/Source/Developer/AITestSuite/Private/Tests/AITestsCommon.h b/Engine/Source/Developer/AITestSuite/Private/Tests/AITestsCommon.h
--- a/Engine/Source/Developer/AITestSuite/Private/Tests/AITestsCommon.h
+++ b/Engine/Source/Developer/AITestSuite/Private/Tests/AITestsCommon.h
@@ -30,6 +30,9 @@ namespace FAITestHelpers
 {
 	UWorld* GetWorld();
 	static const float TickInterval = 1.f / 30;
+
+	void UpdateFrameCounter();
+	uint64 FramesCounter();
 }
 
 struct FAITestBase
