Fixed a crash in automated AI tests when reporting difference between expected and actual results #UE4

diff --git a/Engine/Source/Developer/AITestSuite/Private/Tests/AITestsCommon.cpp b/Engine/Source/Developer/AITestSuite/Private/Tests/AITestsCommon.cpp
--- a/Engine/Source/Developer/AITestSuite/Private/Tests/AITestsCommon.cpp
+++ b/Engine/Source/Developer/AITestSuite/Private/Tests/AITestsCommon.cpp
@@ -142,26 +142,27 @@ void FAITest_SimpleBT::VerifyResults()
 	const bool bMatch = (ExpectedResult == UMockAI_BT::ExecutionLog);
 	if (!bMatch)
 	{
-		FString Desc;
+		FString DescriptionResult;
 		for (int32 Idx = 0; Idx < UMockAI_BT::ExecutionLog.Num(); Idx++)
 		{
-			Desc += TTypeToString<int32>::ToString(UMockAI_BT::ExecutionLog[Idx]);
+			DescriptionResult += TTypeToString<int32>::ToString(UMockAI_BT::ExecutionLog[Idx]);
 			if (Idx < (UMockAI_BT::ExecutionLog.Num() - 1))
 			{
-				Desc += TEXT(", ");
+				DescriptionResult += TEXT(", ");
 			}
 		}
 
+		FString DescriptionExpected;
 		for (int32 Idx = 0; Idx < ExpectedResult.Num(); Idx++)
 		{
-			Desc += TTypeToString<int32>::ToString(ExpectedResult[Idx]);
+			DescriptionExpected += TTypeToString<int32>::ToString(ExpectedResult[Idx]);
 			if (Idx < (ExpectedResult.Num() - 1))
 			{
-				Desc += TEXT(", ");
+				DescriptionExpected += TEXT(", ");
 			}
 		}
 
-		UE_LOG(LogBehaviorTreeTest, Error, TEXT("Test scenario failed to produce expected results!\nExecution log: %s\nExpected values: %s"), *Desc);
+		UE_LOG(LogBehaviorTreeTest, Error, TEXT("Test scenario failed to produce expected results!\nExecution log: %s\nExpected values: %s"), *DescriptionResult, *DescriptionResult);
 	}
 }
 
