Fixed wrong sumary being displayed for failed AI unit tests #UE4

diff --git a/Engine/Source/Developer/AITestSuite/Private/Tests/AITestsCommon.cpp b/Engine/Source/Developer/AITestSuite/Private/Tests/AITestsCommon.cpp
--- a/Engine/Source/Developer/AITestSuite/Private/Tests/AITestsCommon.cpp
+++ b/Engine/Source/Developer/AITestSuite/Private/Tests/AITestsCommon.cpp
@@ -162,7 +162,7 @@ void FAITest_SimpleBT::VerifyResults()
 			}
 		}
 
-		UE_LOG(LogBehaviorTreeTest, Error, TEXT("Test scenario failed to produce expected results!\nExecution log: %s\nExpected values: %s"), *DescriptionResult, *DescriptionResult);
+		UE_LOG(LogBehaviorTreeTest, Error, TEXT("Test scenario failed to produce expected results!\nExecution log: %s\nExpected values: %s"), *DescriptionResult, *DescriptionExpected);
 	}
 }
 
