moved behavior tree execution updates to Tick function to improve automation timings, adjusted existing tests

diff --git a/Engine/Source/Developer/AITestSuite/Private/Tests/BTTest.cpp b/Engine/Source/Developer/AITestSuite/Private/Tests/BTTest.cpp
--- a/Engine/Source/Developer/AITestSuite/Private/Tests/BTTest.cpp
+++ b/Engine/Source/Developer/AITestSuite/Private/Tests/BTTest.cpp
@@ -101,9 +101,9 @@ struct FAITest_BTBasicParallelAbort : public FAITest_SimpleBT
 
 				UBTCompositeNode& CompNode3 = FBTBuilder::AddSequence(CompNode2);
 				{
-					FBTBuilder::AddTaskLogFinish(CompNode3, 1, 11, EBTNodeResult::Succeeded, 3);
-					FBTBuilder::AddTaskLogFinish(CompNode3, 2, 12, EBTNodeResult::Succeeded, 3);
-					FBTBuilder::AddTaskLogFinish(CompNode3, 3, 13, EBTNodeResult::Succeeded, 3);
+					FBTBuilder::AddTaskLogFinish(CompNode3, 1, 11, EBTNodeResult::Succeeded, 4);
+					FBTBuilder::AddTaskLogFinish(CompNode3, 2, 12, EBTNodeResult::Succeeded, 4);
+					FBTBuilder::AddTaskLogFinish(CompNode3, 3, 13, EBTNodeResult::Succeeded, 4);
 				}
 			}
 
@@ -349,7 +349,7 @@ struct FAITest_BTAbortParallelInternal : public FAITest_SimpleBT
 
 					UBTCompositeNode& CompNode4 = FBTBuilder::AddSelector(CompNode3);
 					{
-						FBTBuilder::AddTask(CompNode4, 2, EBTNodeResult::Succeeded, 1);
+						FBTBuilder::AddTask(CompNode4, 2, EBTNodeResult::Succeeded, 3);
 						{
 							FBTBuilder::WithDecoratorBlackboard(CompNode4, EBasicKeyOperation::Set, EBTFlowAbortMode::LowerPriority);
 						}
@@ -406,45 +406,40 @@ struct FAITest_BTAbortParallelOut : public FAITest_SimpleBT
 };
 IMPLEMENT_AI_TEST(FAITest_BTAbortParallelOut, "Engine.AI.Behavior Trees.Abort: parallel out")
 
-//struct FAITest_BTAbortParallelOutAndBack : public FAITest_SimpleBT
-//{
-//	FAITest_BTAbortParallelOutAndBack()
-//	{
-//		UBTCompositeNode& CompNode = FBTBuilder::AddSelector(*BTAsset);
-//		{
-//			FBTBuilder::AddTask(CompNode, 0, EBTNodeResult::Succeeded);
-//			{
-//				FBTBuilder::WithDecoratorBlackboard(CompNode, EBasicKeyOperation::Set, EBTFlowAbortMode::LowerPriority);
-//				FBTBuilder::WithDecorator<UTestBTDecorator_CantExecute>(CompNode);
-//			}
-//
-//			UBTCompositeNode& CompNode2 = FBTBuilder::AddParallel(CompNode, EBTParallelMode::WaitForBackground);
-//			{
-//				FBTBuilder::AddTask(CompNode2, 1, EBTNodeResult::Failed, 5);
-//
-//				UBTCompositeNode& CompNode3 = FBTBuilder::AddSequence(CompNode2);
-//				{
-//					FBTBuilder::AddTask(CompNode3, 2, EBTNodeResult::Succeeded, 1);
-//					FBTBuilder::AddTaskFlagChange(CompNode3, true, EBTNodeResult::Succeeded);
-//					FBTBuilder::AddTask(CompNode3, 3, EBTNodeResult::Succeeded, 1);
-//				}
-//			}
-//
-//			FBTBuilder::AddTask(CompNode, 4, EBTNodeResult::Succeeded);
-//		}
-//
-//		ExpectedResult.Add(1);
-//		ExpectedResult.Add(2);
-//		ExpectedResult.Add(3);
-//		// added this as actual execution shows that's the case
-//		// needs to be reviewed
-//		ExpectedResult.Add(2);
-//		ExpectedResult.Add(3);
-//		//
-//		ExpectedResult.Add(4);
-//	}
-//};
-//IMPLEMENT_AI_TEST(FAITest_BTAbortParallelOutAndBack, "Engine.AI.Behavior Trees.Abort: parallel out & back")
+struct FAITest_BTAbortParallelOutAndBack : public FAITest_SimpleBT
+{
+	FAITest_BTAbortParallelOutAndBack()
+	{
+		UBTCompositeNode& CompNode = FBTBuilder::AddSelector(*BTAsset);
+		{
+			FBTBuilder::AddTask(CompNode, 0, EBTNodeResult::Succeeded);
+			{
+				FBTBuilder::WithDecoratorBlackboard(CompNode, EBasicKeyOperation::Set, EBTFlowAbortMode::LowerPriority);
+				FBTBuilder::WithDecorator<UTestBTDecorator_CantExecute>(CompNode);
+			}
+
+			UBTCompositeNode& CompNode2 = FBTBuilder::AddParallel(CompNode, EBTParallelMode::WaitForBackground);
+			{
+				FBTBuilder::AddTask(CompNode2, 1, EBTNodeResult::Failed, 5);
+
+				UBTCompositeNode& CompNode3 = FBTBuilder::AddSequence(CompNode2);
+				{
+					FBTBuilder::AddTask(CompNode3, 2, EBTNodeResult::Succeeded, 2);
+					FBTBuilder::AddTaskFlagChange(CompNode3, true, EBTNodeResult::Succeeded);
+					FBTBuilder::AddTask(CompNode3, 3, EBTNodeResult::Succeeded, 3);
+				}
+			}
+
+			FBTBuilder::AddTask(CompNode, 4, EBTNodeResult::Succeeded);
+		}
+
+		ExpectedResult.Add(1);
+		ExpectedResult.Add(2);
+		ExpectedResult.Add(3);
+		ExpectedResult.Add(4);
+	}
+};
+IMPLEMENT_AI_TEST(FAITest_BTAbortParallelOutAndBack, "Engine.AI.Behavior Trees.Abort: parallel out & back")
 
 struct FAITest_BTAbortMultipleDelayed : public FAITest_SimpleBT
 {
diff --git a/Engine/Source/Runtime/AIModule/Private/BehaviorTree/BehaviorTreeComponent.cpp b/Engine/Source/Runtime/AIModule/Private/BehaviorTree/BehaviorTreeComponent.cpp
--- a/Engine/Source/Runtime/AIModule/Private/BehaviorTree/BehaviorTreeComponent.cpp
+++ b/Engine/Source/Runtime/AIModule/Private/BehaviorTree/BehaviorTreeComponent.cpp
@@ -168,13 +168,7 @@ void UBehaviorTreeComponent::StopTree()
 	PendingExecution = FBTPendingExecutionInfo();
 	ActiveInstanceIdx = 0;
 
-	if (IsRegistered())
-	{
-		check(GetWorld() != NULL);
-		GetWorld()->GetTimerManager().ClearTimer(this, &UBehaviorTreeComponent::ProcessExecutionRequest);
-	}
-
-	// make sure to allow new execution requests since we just removed last request timer
+	// make sure to allow new execution requests
 	bRequestedFlowUpdate = false;
 }
 
@@ -495,12 +489,7 @@ static void FindCommonParent(const TArray<FBehaviorTreeInstance>& Instances, con
 
 void UBehaviorTreeComponent::ScheduleExecutionUpdate()
 {
-	if (!bRequestedFlowUpdate && IsRegistered())
-	{
-		bRequestedFlowUpdate = true;
-
-		GetWorld()->GetTimerManager().SetTimer(this, &UBehaviorTreeComponent::ProcessExecutionRequest, 0.001f, false);
-	}
+	bRequestedFlowUpdate = true;
 }
 
 void UBehaviorTreeComponent::RequestExecution(UBTCompositeNode* RequestedOn, int32 InstanceIdx, const UBTNode* RequestedBy,
@@ -791,11 +780,16 @@ void UBehaviorTreeComponent::TickComponent(float DeltaTime, enum ELevelTick Tick
 
 	check(this && this->IsPendingKill() == false);
 		
+	if (bRequestedFlowUpdate)
+	{
+		ProcessExecutionRequest();
+	}
+	
 	if (InstanceStack.Num() == 0 || !bIsRunning)
 	{
 		return;
 	}
-	
+
 	// tick active auxiliary nodes and parallel tasks (in execution order, before task)
 	for (int32 InstanceIndex = 0; InstanceIndex < InstanceStack.Num(); InstanceIndex++)
 	{
