removed ensure, it should always show up in logs

diff --git a/Engine/Source/Developer/AITestSuite/Private/Tests/AITestsCommon.cpp b/Engine/Source/Developer/AITestSuite/Private/Tests/AITestsCommon.cpp
--- a/Engine/Source/Developer/AITestSuite/Private/Tests/AITestsCommon.cpp
+++ b/Engine/Source/Developer/AITestSuite/Private/Tests/AITestsCommon.cpp
@@ -150,7 +150,6 @@ bool FAITest_SimpleBT::Update()
 void FAITest_SimpleBT::VerifyResults()
 {
 	const bool bMatch = (ExpectedResult == UMockAI_BT::ExecutionLog);
-	ensure(bMatch && "VerifyResults failed!");
 	if (!bMatch)
 	{
 		FString DescriptionResult;
