fixed functional test for behavior tree subtrees

diff --git a/Engine/Source/Developer/AITestSuite/Public/BTBuilder.h b/Engine/Source/Developer/AITestSuite/Public/BTBuilder.h
--- a/Engine/Source/Developer/AITestSuite/Public/BTBuilder.h
+++ b/Engine/Source/Developer/AITestSuite/Public/BTBuilder.h
@@ -37,6 +37,8 @@ struct FBTBuilder
 		KeyData.KeyType = NewObject<UBlackboardKeyType_Bool>();
 		BB->Keys.Add(KeyData);
 
+		BB->UpdateParentKeys();
+
 		UBehaviorTree* TreeOb = NewObject<UBehaviorTree>();
 		TreeOb->BlackboardAsset = BB;
 
