Mac CIS fix
#mac

diff --git a/Engine/Source/Developer/AITestSuite/Private/Tests/AITestsCommon.h b/Engine/Source/Developer/AITestSuite/Private/Tests/AITestsCommon.h
--- a/Engine/Source/Developer/AITestSuite/Private/Tests/AITestsCommon.h
+++ b/Engine/Source/Developer/AITestSuite/Private/Tests/AITestsCommon.h
@@ -44,7 +44,7 @@ private:
 protected:
 	FAutomationTestBase* TestRunner;
 
-	FAITestBase() : TestRunner(nullptr), bTearedDown(false)
+	FAITestBase() : bTearedDown(false), TestRunner(nullptr)
 	{}
 
 	template<typename ClassToSpawn>
