Fixed Automation crashes on AITestSuite's tests #UE4

diff --git a/Engine/Source/Developer/AITestSuite/Private/Tests/AITestsCommon.cpp b/Engine/Source/Developer/AITestSuite/Private/Tests/AITestsCommon.cpp
--- a/Engine/Source/Developer/AITestSuite/Private/Tests/AITestsCommon.cpp
+++ b/Engine/Source/Developer/AITestSuite/Private/Tests/AITestsCommon.cpp
@@ -116,6 +116,8 @@ void FAITestBase::TearDown()
 FAITest_SimpleBT::FAITest_SimpleBT()
 {
 	bUseSystemTicking = false;
+	
+	BTAsset = &FBTBuilder::CreateBehaviorTree();
 }
 
 void FAITest_SimpleBT::SetUp()
@@ -123,7 +125,6 @@ void FAITest_SimpleBT::SetUp()
 	FAITestBase::SetUp();
 
 	AIBTUser = NewAutoDestroyObject<UMockAI_BT>();
-	BTAsset = &FBTBuilder::CreateBehaviorTree();
 
 	UMockAI_BT::ExecutionLog.Reset();
 
@@ -139,7 +140,6 @@ bool FAITest_SimpleBT::Update()
 	FAITestHelpers::UpdateFrameCounter();
 
 	if (AIBTUser != NULL)
-
 	{
 		if (bUseSystemTicking == false)
 		{
@@ -159,7 +159,7 @@ bool FAITest_SimpleBT::Update()
 void FAITest_SimpleBT::VerifyResults()
 {
 	const bool bMatch = (ExpectedResult == UMockAI_BT::ExecutionLog);
-	ensure(bMatch && "VerifyResults failed!");
+	//ensure(bMatch && "VerifyResults failed!");
 	if (!bMatch)
 	{
 		FString DescriptionResult;
