Fix categorization of AI test added by merge.

diff --git a/Engine/Source/Developer/AITestSuite/Private/Tests/BTTest.cpp b/Engine/Source/Developer/AITestSuite/Private/Tests/BTTest.cpp
--- a/Engine/Source/Developer/AITestSuite/Private/Tests/BTTest.cpp
+++ b/Engine/Source/Developer/AITestSuite/Private/Tests/BTTest.cpp
@@ -573,7 +573,7 @@ struct FAITest_BTLowPriObserverInLoop : public FAITest_SimpleBT
 		ExpectedResult.Add(3);
 	}
 };
-IMPLEMENT_AI_LATENT_TEST(FAITest_BTLowPriObserverInLoop, "Engine.AI.Behavior Trees.Other: low priority observer in looped branch")
+IMPLEMENT_AI_LATENT_TEST(FAITest_BTLowPriObserverInLoop, "System.Engine.AI.Behavior Trees.Other: low priority observer in looped branch")
 
 struct FAITest_BTSubtreeSimple : public FAITest_SimpleBT
 {
