Missing GameplayTasks changes merged over from dev branch #UE4

diff --git a/Engine/Source/Developer/AITestSuite/Classes/MockGameplayTasks.h b/Engine/Source/Developer/AITestSuite/Classes/MockGameplayTasks.h
--- a/Engine/Source/Developer/AITestSuite/Classes/MockGameplayTasks.h
+++ b/Engine/Source/Developer/AITestSuite/Classes/MockGameplayTasks.h
@@ -3,6 +3,7 @@
 #pragma once 
 
 #include "GameplayTask.h"
+#include "GameplayTasksComponent.h"
 #include "MockGameplayTasks.generated.h"
 
 namespace ETestTaskMessage
@@ -28,7 +29,7 @@ protected:
 public:
 	UMockTask_Log(const FObjectInitializer& ObjectInitializer);
 
-	static UMockTask_Log* CreateTask(IGameplayTaskOwnerInterface& TaskOwner, FTestLogger<int32>& InLogger);
+	static UMockTask_Log* CreateTask(IGameplayTaskOwnerInterface& TaskOwner, FTestLogger<int32>& InLogger, const FGameplayResourceSet& Resources = FGameplayResourceSet(), uint8 Priority = FGameplayTasks::DefaultPriority);
 
 protected:
 	virtual void Activate() override;
@@ -42,3 +43,15 @@ public:
 	// testing only hack-functions
 	void EnableTick() { bTickingTask = true; }
 };
+
+//
+// a Testing-time component that is a way to access UGameplayTasksComponent's protected properties
+//
+UCLASS()
+class UMockGameplayTasksComponent : public UGameplayTasksComponent
+{
+	GENERATED_BODY()
+
+public:	
+	int32 GetTaskPriorityQueueSize() const { return TaskPriorityQueue.Num(); }
+};
diff --git a/Engine/Source/Developer/AITestSuite/Private/MockAI/MockGameplayTasks.cpp b/Engine/Source/Developer/AITestSuite/Private/MockAI/MockGameplayTasks.cpp
--- a/Engine/Source/Developer/AITestSuite/Private/MockAI/MockGameplayTasks.cpp
+++ b/Engine/Source/Developer/AITestSuite/Private/MockAI/MockGameplayTasks.cpp
@@ -12,12 +12,14 @@ UMockTask_Log::UMockTask_Log(const FObjectInitializer& ObjectInitializer)
 
 }
 
-UMockTask_Log* UMockTask_Log::CreateTask(IGameplayTaskOwnerInterface& TaskOwner, FTestLogger<int32>& InLogger)
+UMockTask_Log* UMockTask_Log::CreateTask(IGameplayTaskOwnerInterface& TaskOwner, FTestLogger<int32>& InLogger, const FGameplayResourceSet& Resources, uint8 Priority)
 {
 	UMockTask_Log* Task = NewTask<UMockTask_Log>(TaskOwner);
 	if (Task)
 	{
 		Task->Logger = &InLogger;
+		Task->RequiredResources = Resources;
+		Task->Priority = Priority;
 	}
 	return Task;
 }
diff --git a/Engine/Source/Developer/AITestSuite/Private/Tests/GameplayTasksTest.cpp b/Engine/Source/Developer/AITestSuite/Private/Tests/GameplayTasksTest.cpp
--- a/Engine/Source/Developer/AITestSuite/Private/Tests/GameplayTasksTest.cpp
+++ b/Engine/Source/Developer/AITestSuite/Private/Tests/GameplayTasksTest.cpp
@@ -9,8 +9,21 @@
 
 #define LOCTEXT_NAMESPACE "AITestSuite_GameplayTasksTest"
 
-typedef FAITest_SimpleComponentBasedTest<UGameplayTasksComponent> FAITest_GameplayTasksTest;
+typedef FAITest_SimpleComponentBasedTest<UMockGameplayTasksComponent> FAITest_GameplayTasksTest;
 
+static const FGameplayResourceSet::FResourceID ResourceMovement = 0;
+static const FGameplayResourceSet::FResourceID ResourceLogic = 1;
+static const FGameplayResourceSet::FResourceID ResourceAnimation = 2;
+
+static const FGameplayResourceSet MovementResourceSet = FGameplayResourceSet().AddID(ResourceMovement);
+static const FGameplayResourceSet LogicResourceSet = FGameplayResourceSet().AddID(ResourceLogic);
+static const FGameplayResourceSet AnimationResourceSet = FGameplayResourceSet().AddID(ResourceAnimation);
+static const FGameplayResourceSet MoveAndAnimResourceSet = FGameplayResourceSet().AddSet(MovementResourceSet).AddSet(AnimationResourceSet);
+static const FGameplayResourceSet MoveAndLogicResourceSet = FGameplayResourceSet((1 << ResourceMovement) | (1 << ResourceLogic));
+static const FGameplayResourceSet MoveAnimLogicResourceSet = FGameplayResourceSet((1 << ResourceMovement) | (1 << ResourceLogic) | (1 << ResourceAnimation));
+
+static const uint8 LowPriority = 1;
+static const uint8 HighPriority = 255;
 
 struct FAITest_GameplayTask_ComponentState : public FAITest_GameplayTasksTest
 {
@@ -104,4 +117,238 @@ struct FAITest_GameplayTask_SelfEnd : public FAITest_GameplayTasksTest
 };
 IMPLEMENT_AI_LATENT_TEST(FAITest_GameplayTask_SelfEnd, "Engine.AI.Gameplay Tasks.Self End")
 
+//----------------------------------------------------------------------//
+// Testing multiple simultaneously ticking tasks
+//----------------------------------------------------------------------//
+struct FAITest_GameplayTask_SimulatanousTick : public FAITest_GameplayTasksTest
+{
+	UMockTask_Log* Tasks[3];
+
+	void SetUp()
+	{
+		FAITest_GameplayTasksTest::SetUp();
+
+		Logger.ExpectedValues.Add(ETestTaskMessage::Activate);
+		Logger.ExpectedValues.Add(ETestTaskMessage::Activate); 
+		Logger.ExpectedValues.Add(ETestTaskMessage::Activate); 
+		Logger.ExpectedValues.Add(ETestTaskMessage::Tick);
+		Logger.ExpectedValues.Add(ETestTaskMessage::Tick);
+		Logger.ExpectedValues.Add(ETestTaskMessage::Tick);
+		Logger.ExpectedValues.Add(ETestTaskMessage::ExternalCancel);
+		Logger.ExpectedValues.Add(ETestTaskMessage::Ended);
+		Logger.ExpectedValues.Add(ETestTaskMessage::ExternalCancel);
+		Logger.ExpectedValues.Add(ETestTaskMessage::Ended);
+		Logger.ExpectedValues.Add(ETestTaskMessage::ExternalCancel);
+		Logger.ExpectedValues.Add(ETestTaskMessage::Ended);
+
+		for (int32 Index = 0; Index < sizeof(Tasks) / sizeof(UMockTask_Log*); ++Index)
+		{
+			Tasks[Index] = UMockTask_Log::CreateTask(*Component, Logger);
+			Tasks[Index]->EnableTick();
+			Tasks[Index]->ReadyForActivation();
+		}
+		Test(TEXT("Component should want to tick in this scenario"), Component->GetShouldTick() == true);
+	}
+
+	bool Update()
+	{
+		TickComponent();
+		for (int32 Index = 0; Index < sizeof(Tasks) / sizeof(UMockTask_Log*); ++Index)
+		{
+			Tasks[Index]->ExternalCancel();
+		}
+		Test(TEXT("Component should want to tick in this scenario"), Component->GetShouldTick() == false);
+
+		return true;
+	}
+};
+IMPLEMENT_AI_LATENT_TEST(FAITest_GameplayTask_SimulatanousTick, "Engine.AI.Gameplay Tasks.Simultanously ticking tasks")
+
+//----------------------------------------------------------------------//
+// Testing multiple simultaneously ticking tasks UGameplayTaskResource
+//----------------------------------------------------------------------//
+struct FAITest_GameplayTask_ResourceSet : public FAITestBase
+{
+	void InstantTest()
+	{
+		FGameplayResourceSet ResourcesSet;
+		Test(TEXT("New FGameplayResourceSet should be empty initialy"), ResourcesSet.IsEmpty() == true);
+		ResourcesSet.AddID(ResourceLogic);
+		Test(TEXT("Added one ID, ResourcesSet should not be perceived as empty now"), ResourcesSet.IsEmpty() == false);
+		Test(TEXT("Added one ID, ResourcesSet should not be perceived as empty now"), ResourcesSet.IsEmpty() == false);
+		ResourcesSet.RemoveID(ResourceAnimation);
+		Test(TEXT("Removed ID not previously added, ResourcesSet should not be perceived as empty now"), ResourcesSet.IsEmpty() == false);
+		ResourcesSet.RemoveID(ResourceLogic);
+		Test(TEXT("Removed ID previously added, ResourcesSet should be empty now"), ResourcesSet.IsEmpty());
+
+		Test(TEXT("Single ID checking, not present ID"), MoveAndAnimResourceSet.HasAnyID(LogicResourceSet) == false);
+		Test(TEXT("Single ID checking"), MoveAndAnimResourceSet.HasAnyID(MovementResourceSet));
+		Test(TEXT("Single ID checking"), MoveAndAnimResourceSet.HasAnyID(AnimationResourceSet));
+
+		Test(TEXT("Multiple ID checking - has all, self test"), MoveAndAnimResourceSet.HasAllIDs(MoveAndAnimResourceSet));
+		Test(TEXT("Multiple ID checking - has all, other identical"), MoveAndAnimResourceSet.HasAllIDs(FGameplayResourceSet((1 << ResourceMovement) | (1 << ResourceAnimation))));
+		Test(TEXT("Multiple ID checking - has all, other different"), MoveAndAnimResourceSet.HasAllIDs(MoveAndLogicResourceSet) == false);
+		Test(TEXT("Multiple ID checking - overlap"), MoveAndAnimResourceSet.GetOverlap(MoveAndLogicResourceSet) == MovementResourceSet);
+		Test(TEXT("Multiple ID checking - substraction"), MoveAndAnimResourceSet.GetDifference(MoveAndLogicResourceSet) == AnimationResourceSet);
+		
+		Test(TEXT("FGameplayResourceSet containing 0-th ID is not empty"), MovementResourceSet.IsEmpty() == false);
+		Test(TEXT("FGameplayResourceSet has 0-th ID"), MovementResourceSet.HasID(ResourceMovement));
+	}
+};
+IMPLEMENT_AI_INSTANT_TEST(FAITest_GameplayTask_ResourceSet, "Engine.AI.Gameplay Tasks.Resource Set")
+
+//----------------------------------------------------------------------//
+// Running tasks requiring non-overlapping resources
+//----------------------------------------------------------------------//
+struct FAITest_GameplayTask_NonOverlappingResources : public FAITest_GameplayTasksTest
+{
+	UMockTask_Log* Tasks[2];
+
+	void SetUp()
+	{
+		UWorld& World = GetWorld();
+		
+		Tasks[0] = UMockTask_Log::CreateTask(*Component, Logger, MoveAndAnimResourceSet);
+		Tasks[1] = UMockTask_Log::CreateTask(*Component, Logger, LogicResourceSet);
+
+		Tasks[0]->ReadyForActivation();
+		Test(TEXT("TasksComponent should claim it's using 0th task's resources"), Component->GetCurrentlyUsedResources() == Tasks[0]->GetRequiredResources());
+		
+		Tasks[1]->ReadyForActivation();		
+		Test(TEXT("Both tasks should be \'Active\' since their resources do not overlap"), Tasks[0]->GetState() == EGameplayTaskState::Active && Tasks[1]->GetState() == EGameplayTaskState::Active);
+		Test(TEXT("TasksComponent should claim it's using only latter task's resources"), Component->GetCurrentlyUsedResources() == MoveAnimLogicResourceSet);
+
+		Tasks[0]->ExternalCancel();
+		Test(TEXT("Only index 1 task's resources should be relevant now"), Component->GetCurrentlyUsedResources() == Tasks[1]->GetRequiredResources());
+
+		Tasks[1]->ExternalCancel();
+		Test(TEXT("No resources should be occupied now"), Component->GetCurrentlyUsedResources().IsEmpty());
+	}
+};
+IMPLEMENT_AI_LATENT_TEST(FAITest_GameplayTask_NonOverlappingResources, "Engine.AI.Gameplay Tasks.Non-overlapping resources")
+
+//----------------------------------------------------------------------//
+// Running tasks requiring overlapping resources
+//----------------------------------------------------------------------//
+struct FAITest_GameplayTask_OverlappingResources : public FAITest_GameplayTasksTest
+{
+	UMockTask_Log* Tasks[2];
+
+	void SetUp()
+	{
+		UWorld& World = GetWorld();
+
+		Tasks[0] = UMockTask_Log::CreateTask(*Component, Logger, MoveAndAnimResourceSet);
+		Tasks[1] = UMockTask_Log::CreateTask(*Component, Logger, MoveAndLogicResourceSet);
+
+		Tasks[0]->ReadyForActivation();
+		Tasks[1]->ReadyForActivation();
+
+		Test(TEXT("Only the latter task should be active since it shadows the other one in terms of required resources"), Tasks[1]->GetState() == EGameplayTaskState::Active);
+		Test(TEXT("The first task should be paused at this moment"), Tasks[0]->GetState() == EGameplayTaskState::Paused);
+		Test(TEXT("TasksComponent should claim it's using only latter task's resources"), Component->GetCurrentlyUsedResources() == Tasks[1]->GetRequiredResources());
+
+		Tasks[1]->ExternalCancel();
+		Test(TEXT("Now the latter task should be marked as Finished"), Tasks[1]->GetState() == EGameplayTaskState::Finished);
+		Test(TEXT("And the first task should be resumed"), Tasks[0]->GetState() == EGameplayTaskState::Active);
+		Test(TEXT("TasksComponent should claim it's using only first task's resources"), Component->GetCurrentlyUsedResources() == Tasks[0]->GetRequiredResources());
+	}
+};
+IMPLEMENT_AI_LATENT_TEST(FAITest_GameplayTask_OverlappingResources, "Engine.AI.Gameplay Tasks.Overlapping resources")
+
+//----------------------------------------------------------------------//
+// Priority handling
+//----------------------------------------------------------------------//
+struct FAITest_GameplayTask_Priorities : public FAITest_GameplayTasksTest
+{
+	UMockTask_Log* Tasks[3];
+
+	void SetUp()
+	{
+		UWorld& World = GetWorld();
+
+		// all tasks use same resources, have different priorities
+		// let's do some tick testing as well
+		Tasks[0] = UMockTask_Log::CreateTask(*Component, Logger, MoveAndAnimResourceSet, LowPriority);
+		Tasks[0]->EnableTick();
+		Tasks[1] = UMockTask_Log::CreateTask(*Component, Logger, MoveAndAnimResourceSet);
+		Tasks[1]->EnableTick();
+		Tasks[2] = UMockTask_Log::CreateTask(*Component, Logger, MoveAndAnimResourceSet, HighPriority);		
+
+		Tasks[1]->ReadyForActivation();
+		Tasks[0]->ReadyForActivation();
+
+		Test(TEXT("Task at index 1 should be active at this point since it's higher priority"), Tasks[1]->IsActive() && !Tasks[0]->IsActive());
+		Test(TEXT("TasksComponent should claim it's using only resources of task 1"), Component->GetCurrentlyUsedResources() == Tasks[1]->GetRequiredResources());
+		Test(TEXT("Current top action wants to tick so Component should want that as well"), Component->GetShouldTick() == true); 
+
+		Tasks[2]->ReadyForActivation();
+		Test(TEXT("Now the last pushed, highest priority task should be active"), Tasks[2]->IsActive() && !Tasks[0]->IsActive() && !Tasks[1]->IsActive());
+		Test(TEXT("No ticking task is active so Component should not want to tick"), Component->GetShouldTick() == false);
+		
+		Tasks[1]->ExternalCancel();
+		Test(TEXT("Canceling mid-priority inactive task should not influence what's active"), Tasks[2]->IsActive() && !Tasks[0]->IsActive() && !Tasks[1]->IsActive());
+		Test(TEXT("Current top action still doesn't want to tick, so neither should the Component"), Component->GetShouldTick() == false);
+
+		Tasks[2]->ExternalCancel();
+		Test(TEXT("After canceling the top-priority task the lowest priority task remains to be active"), !Tasks[2]->IsActive() && Tasks[0]->IsActive() && !Tasks[1]->IsActive());
+		Test(TEXT("New top action wants tick, so should Component"), Component->GetShouldTick() == true);
+
+		Tasks[0]->ExternalCancel();
+		Test(TEXT("Task-less component should not want to tick"), Component->GetShouldTick() == false);
+	}
+};
+IMPLEMENT_AI_LATENT_TEST(FAITest_GameplayTask_Priorities, "Engine.AI.Gameplay Tasks.Priorities")
+
+//----------------------------------------------------------------------//
+// Internal ending, by task ending itself or owner finishing 
+//----------------------------------------------------------------------//
+struct FAITest_GameplayTask_InternalEnding : public FAITest_GameplayTasksTest
+{
+	static const int32 TasksCount = 4;
+	UMockTask_Log* Task_SelfEndNoResources;
+	UMockTask_Log* Task_OwnerEndNoResources;
+	UMockTask_Log* Task_SelfEndWithResources;
+	UMockTask_Log* Task_OwnerEndWithResources;
+	UMockTask_Log* Tasks[TasksCount];
+
+	void SetUp()
+	{
+		UWorld& World = GetWorld();
+
+		// all tasks use same resources, have different priorities
+		Tasks[0] = Task_SelfEndNoResources = UMockTask_Log::CreateTask(*Component, Logger);
+		Tasks[1] = Task_OwnerEndNoResources = UMockTask_Log::CreateTask(*Component, Logger);
+		// not using overlapping resources set on purpose - want to test them independently
+		Tasks[2] = Task_SelfEndWithResources = UMockTask_Log::CreateTask(*Component, Logger, MovementResourceSet);
+		Tasks[3] = Task_OwnerEndWithResources = UMockTask_Log::CreateTask(*Component, Logger, LogicResourceSet);
+		
+		for (int32 Index = 0; Index < TasksCount; ++Index)
+		{
+			Tasks[Index]->ReadyForActivation();
+			Test(TEXT("Trivial activation should succeed"), Tasks[Index]->IsActive());
+		}
+		Test(TEXT("Resources should sum up"), Component->GetCurrentlyUsedResources() == MoveAndLogicResourceSet);
+
+		Task_SelfEndNoResources->EndTask();
+		Test(TEXT("Task_SelfEndNoResources should be \'done\' now"), Task_SelfEndNoResources->IsActive() == false);
+
+		Task_OwnerEndNoResources->TaskOwnerEnded();
+		Test(TEXT("Task_SelfEndNoResources should be \'done\' now"), Task_OwnerEndNoResources->IsActive() == false);
+
+		Task_SelfEndWithResources->EndTask();
+		Test(TEXT("Task_SelfEndWithResources should be \'done\' now"), Task_SelfEndWithResources->IsActive() == false);
+		Test(TEXT("Only the other task's resources should matter now"), Component->GetCurrentlyUsedResources() == LogicResourceSet);
+		Test(TEXT("There should be only one active task in the priority queue"), Component->GetTaskPriorityQueueSize() == 1);
+
+		Task_OwnerEndWithResources->TaskOwnerEnded();
+		Test(TEXT("Task_SelfEndWithResources should be \'done\' now"), Task_OwnerEndWithResources->IsActive() == false);
+		Test(TEXT("No resources should be locked at this moment"), Component->GetCurrentlyUsedResources().IsEmpty());
+		Test(TEXT("Priority Task Queue should be empty"), Component->GetTaskPriorityQueueSize() == 0);
+	}
+};
+IMPLEMENT_AI_LATENT_TEST(FAITest_GameplayTask_InternalEnding, "Engine.AI.Gameplay Tasks.Self and Owner ending")
+
+// add tests if component wants ticking at while aborting/reactivating tasks
+
 #undef LOCTEXT_NAMESPACE
\ No newline at end of file
