Updated Gameplay Tasks's unit tests category #UE4

Now in System.Engine.AI

diff --git a/Engine/Source/Developer/AITestSuite/Private/Tests/GameplayTasksTest.cpp b/Engine/Source/Developer/AITestSuite/Private/Tests/GameplayTasksTest.cpp
--- a/Engine/Source/Developer/AITestSuite/Private/Tests/GameplayTasksTest.cpp
+++ b/Engine/Source/Developer/AITestSuite/Private/Tests/GameplayTasksTest.cpp
@@ -47,7 +47,7 @@ struct FAITest_GameplayTask_ComponentState : public FAITest_GameplayTasksTest
 		Test(TEXT("Task should be \'uninitialized\' before Activate is called on it"), Task->GetState() == EGameplayTaskState::Active);*/
 	}
 };
-IMPLEMENT_AI_LATENT_TEST(FAITest_GameplayTask_ComponentState, "Engine.AI.Gameplay Tasks.Component\'s basic behavior")
+IMPLEMENT_AI_LATENT_TEST(FAITest_GameplayTask_ComponentState, "System.Engine.AI.Gameplay Tasks.Component\'s basic behavior")
 
 //----------------------------------------------------------------------//
 // 
@@ -85,7 +85,7 @@ struct FAITest_GameplayTask_ExternalCancelWithTick : public FAITest_GameplayTask
 		return true;
 	}
 };
-IMPLEMENT_AI_LATENT_TEST(FAITest_GameplayTask_ExternalCancelWithTick, "Engine.AI.Gameplay Tasks.External Cancel with Tick")
+IMPLEMENT_AI_LATENT_TEST(FAITest_GameplayTask_ExternalCancelWithTick, "System.Engine.AI.Gameplay Tasks.External Cancel with Tick")
 
 //----------------------------------------------------------------------//
 // In this test the task should get properly created, acticated and end 
@@ -115,7 +115,7 @@ struct FAITest_GameplayTask_SelfEnd : public FAITest_GameplayTasksTest
 		return true;
 	}
 };
-IMPLEMENT_AI_LATENT_TEST(FAITest_GameplayTask_SelfEnd, "Engine.AI.Gameplay Tasks.Self End")
+IMPLEMENT_AI_LATENT_TEST(FAITest_GameplayTask_SelfEnd, "System.Engine.AI.Gameplay Tasks.Self End")
 
 //----------------------------------------------------------------------//
 // Testing multiple simultaneously ticking tasks
@@ -162,7 +162,7 @@ struct FAITest_GameplayTask_SimulatanousTick : public FAITest_GameplayTasksTest
 		return true;
 	}
 };
-IMPLEMENT_AI_LATENT_TEST(FAITest_GameplayTask_SimulatanousTick, "Engine.AI.Gameplay Tasks.Simultanously ticking tasks")
+IMPLEMENT_AI_LATENT_TEST(FAITest_GameplayTask_SimulatanousTick, "System.Engine.AI.Gameplay Tasks.Simultanously ticking tasks")
 
 //----------------------------------------------------------------------//
 // Testing multiple simultaneously ticking tasks UGameplayTaskResource
@@ -195,7 +195,7 @@ struct FAITest_GameplayTask_ResourceSet : public FAITestBase
 		Test(TEXT("FGameplayResourceSet has 0-th ID"), MovementResourceSet.HasID(ResourceMovement));
 	}
 };
-IMPLEMENT_AI_INSTANT_TEST(FAITest_GameplayTask_ResourceSet, "Engine.AI.Gameplay Tasks.Resource Set")
+IMPLEMENT_AI_INSTANT_TEST(FAITest_GameplayTask_ResourceSet, "System.Engine.AI.Gameplay Tasks.Resource Set")
 
 //----------------------------------------------------------------------//
 // Running tasks requiring non-overlapping resources
@@ -225,7 +225,7 @@ struct FAITest_GameplayTask_NonOverlappingResources : public FAITest_GameplayTas
 		Test(TEXT("No resources should be occupied now"), Component->GetCurrentlyUsedResources().IsEmpty());
 	}
 };
-IMPLEMENT_AI_LATENT_TEST(FAITest_GameplayTask_NonOverlappingResources, "Engine.AI.Gameplay Tasks.Non-overlapping resources")
+IMPLEMENT_AI_LATENT_TEST(FAITest_GameplayTask_NonOverlappingResources, "System.Engine.AI.Gameplay Tasks.Non-overlapping resources")
 
 //----------------------------------------------------------------------//
 // Running tasks requiring overlapping resources
@@ -254,7 +254,7 @@ struct FAITest_GameplayTask_OverlappingResources : public FAITest_GameplayTasksT
 		Test(TEXT("TasksComponent should claim it's using only first task's resources"), Component->GetCurrentlyUsedResources() == Tasks[0]->GetRequiredResources());
 	}
 };
-IMPLEMENT_AI_LATENT_TEST(FAITest_GameplayTask_OverlappingResources, "Engine.AI.Gameplay Tasks.Overlapping resources")
+IMPLEMENT_AI_LATENT_TEST(FAITest_GameplayTask_OverlappingResources, "System.Engine.AI.Gameplay Tasks.Overlapping resources")
 
 //----------------------------------------------------------------------//
 // Priority handling
@@ -298,7 +298,7 @@ struct FAITest_GameplayTask_Priorities : public FAITest_GameplayTasksTest
 		Test(TEXT("Task-less component should not want to tick"), Component->GetShouldTick() == false);
 	}
 };
-IMPLEMENT_AI_LATENT_TEST(FAITest_GameplayTask_Priorities, "Engine.AI.Gameplay Tasks.Priorities")
+IMPLEMENT_AI_LATENT_TEST(FAITest_GameplayTask_Priorities, "System.Engine.AI.Gameplay Tasks.Priorities")
 
 //----------------------------------------------------------------------//
 // Internal ending, by task ending itself or owner finishing 
@@ -347,7 +347,7 @@ struct FAITest_GameplayTask_InternalEnding : public FAITest_GameplayTasksTest
 		Test(TEXT("Priority Task Queue should be empty"), Component->GetTaskPriorityQueueSize() == 0);
 	}
 };
-IMPLEMENT_AI_LATENT_TEST(FAITest_GameplayTask_InternalEnding, "Engine.AI.Gameplay Tasks.Self and Owner ending")
+IMPLEMENT_AI_LATENT_TEST(FAITest_GameplayTask_InternalEnding, "System.Engine.AI.Gameplay Tasks.Self and Owner ending")
 
 // add tests if component wants ticking at while aborting/reactivating tasks
 
