Gameplay Task priority and resources management up and running, pass 1 #UE4

Including:
- every GameplayTask can specify resources it requires to run, and separately resources it's going to lock, or "shadow" ones in the priority queue
- added a broadcast delegate to GameplayTasksComponent one can register for to get notified about changes in resource locks
    - AI uses this to lock AI logic for scripted tasks
- made IGameplayTaskOwnerInterface::GetAvatarActor require a task pointer to determine avatar actor - this was needed for a "skeletal" owners, like BT nodes
- Improvements to AITask_MoveTo
- expanded BTTaskNode with bits required for it being a GameplayTask owner
    - might need to move that to the most generic BTNode so that services can easily spawn tasks as well
- BT movement tasks can now use AITasks (opt-in with AISystem::bEnableBTAITasks via Project Settings)
- started adding a generic mechanism for configuring AI properties via Project Settings
    - consistency throughout multiple ways of spawning given AI task being the main drive here. Started with movement of course
    - currently implemented by exposing AISystem's properties to global config. Going to switch that over to UDeveloperSettings soon.
- Added a blackboard component member to AIController
- added setting BB SelfActor key to AI-controlled Pawn

diff --git a/Engine/Source/Developer/AITestSuite/Classes/MockGameplayTasks.h b/Engine/Source/Developer/AITestSuite/Classes/MockGameplayTasks.h
--- a/Engine/Source/Developer/AITestSuite/Classes/MockGameplayTasks.h
+++ b/Engine/Source/Developer/AITestSuite/Classes/MockGameplayTasks.h
@@ -55,3 +55,19 @@ class UMockGameplayTasksComponent : public UGameplayTasksComponent
 public:	
 	int32 GetTaskPriorityQueueSize() const { return TaskPriorityQueue.Num(); }
 };
+
+UCLASS()
+class UMockGameplayTaskOwner : public UObject, public IGameplayTaskOwnerInterface
+{
+	GENERATED_BODY()
+
+public:
+	UPROPERTY()
+	UGameplayTasksComponent* GTComponent;
+
+	virtual UGameplayTasksComponent* GetGameplayTasksComponent(const UGameplayTask& Task) const override { return GTComponent; }	
+	virtual void OnTaskActivated(UGameplayTask& Task) override {}
+	virtual void OnTaskDeactivated(UGameplayTask& Task) override {}
+	virtual AActor* GetOwnerActor(const UGameplayTask* Task) const { return nullptr; }
+
+};
\ No newline at end of file
diff --git a/Engine/Source/Developer/AITestSuite/Private/Tests/GameplayTasksTest.cpp b/Engine/Source/Developer/AITestSuite/Private/Tests/GameplayTasksTest.cpp
--- a/Engine/Source/Developer/AITestSuite/Private/Tests/GameplayTasksTest.cpp
+++ b/Engine/Source/Developer/AITestSuite/Private/Tests/GameplayTasksTest.cpp
@@ -27,7 +27,7 @@ static const uint8 HighPriority = 255;
 
 struct FAITest_GameplayTask_ComponentState : public FAITest_GameplayTasksTest
 {
-	void SetUp()
+	void InstantTest()
 	{
 		FAITest_GameplayTasksTest::SetUp();
 
@@ -47,7 +47,7 @@ struct FAITest_GameplayTask_ComponentState : public FAITest_GameplayTasksTest
 		Test(TEXT("Task should be \'uninitialized\' before Activate is called on it"), Task->GetState() == EGameplayTaskState::Active);*/
 	}
 };
-IMPLEMENT_AI_LATENT_TEST(FAITest_GameplayTask_ComponentState, "System.Engine.AI.Gameplay Tasks.Component\'s basic behavior")
+IMPLEMENT_AI_INSTANT_TEST(FAITest_GameplayTask_ComponentState, "System.Engine.AI.Gameplay Tasks.Component\'s basic behavior")
 
 //----------------------------------------------------------------------//
 // 
@@ -204,7 +204,7 @@ struct FAITest_GameplayTask_NonOverlappingResources : public FAITest_GameplayTas
 {
 	UMockTask_Log* Tasks[2];
 
-	void SetUp()
+	void InstantTest()
 	{
 		UWorld& World = GetWorld();
 		
@@ -212,20 +212,20 @@ struct FAITest_GameplayTask_NonOverlappingResources : public FAITest_GameplayTas
 		Tasks[1] = UMockTask_Log::CreateTask(*Component, Logger, LogicResourceSet);
 
 		Tasks[0]->ReadyForActivation();
-		Test(TEXT("TasksComponent should claim it's using 0th task's resources"), Component->GetCurrentlyUsedResources() == Tasks[0]->GetRequiredResources());
+		Test(TEXT("TasksComponent should claim it's using 0th task's resources"), Component->GetCurrentlyUsedResources() == Tasks[0]->GetClaimedResources());
 		
 		Tasks[1]->ReadyForActivation();		
 		Test(TEXT("Both tasks should be \'Active\' since their resources do not overlap"), Tasks[0]->GetState() == EGameplayTaskState::Active && Tasks[1]->GetState() == EGameplayTaskState::Active);
 		Test(TEXT("TasksComponent should claim it's using only latter task's resources"), Component->GetCurrentlyUsedResources() == MoveAnimLogicResourceSet);
 
 		Tasks[0]->ExternalCancel();
-		Test(TEXT("Only index 1 task's resources should be relevant now"), Component->GetCurrentlyUsedResources() == Tasks[1]->GetRequiredResources());
+		Test(TEXT("Only index 1 task's resources should be relevant now"), Component->GetCurrentlyUsedResources() == Tasks[1]->GetClaimedResources());
 
 		Tasks[1]->ExternalCancel();
 		Test(TEXT("No resources should be occupied now"), Component->GetCurrentlyUsedResources().IsEmpty());
 	}
 };
-IMPLEMENT_AI_LATENT_TEST(FAITest_GameplayTask_NonOverlappingResources, "System.Engine.AI.Gameplay Tasks.Non-overlapping resources")
+IMPLEMENT_AI_INSTANT_TEST(FAITest_GameplayTask_NonOverlappingResources, "System.Engine.AI.Gameplay Tasks.Non-overlapping resources")
 
 //----------------------------------------------------------------------//
 // Running tasks requiring overlapping resources
@@ -234,7 +234,7 @@ struct FAITest_GameplayTask_OverlappingResources : public FAITest_GameplayTasksT
 {
 	UMockTask_Log* Tasks[2];
 
-	void SetUp()
+	void InstantTest()
 	{
 		UWorld& World = GetWorld();
 
@@ -246,15 +246,47 @@ struct FAITest_GameplayTask_OverlappingResources : public FAITest_GameplayTasksT
 
 		Test(TEXT("Only the latter task should be active since it shadows the other one in terms of required resources"), Tasks[1]->GetState() == EGameplayTaskState::Active);
 		Test(TEXT("The first task should be paused at this moment"), Tasks[0]->GetState() == EGameplayTaskState::Paused);
-		Test(TEXT("TasksComponent should claim it's using only latter task's resources"), Component->GetCurrentlyUsedResources() == Tasks[1]->GetRequiredResources());
+		Test(TEXT("TasksComponent should claim it's using only latter task's resources"), Component->GetCurrentlyUsedResources() == Tasks[1]->GetClaimedResources());
 
 		Tasks[1]->ExternalCancel();
 		Test(TEXT("Now the latter task should be marked as Finished"), Tasks[1]->GetState() == EGameplayTaskState::Finished);
 		Test(TEXT("And the first task should be resumed"), Tasks[0]->GetState() == EGameplayTaskState::Active);
-		Test(TEXT("TasksComponent should claim it's using only first task's resources"), Component->GetCurrentlyUsedResources() == Tasks[0]->GetRequiredResources());
+		Test(TEXT("TasksComponent should claim it's using only first task's resources"), Component->GetCurrentlyUsedResources() == Tasks[0]->GetClaimedResources());
 	}
 };
-IMPLEMENT_AI_LATENT_TEST(FAITest_GameplayTask_OverlappingResources, "System.Engine.AI.Gameplay Tasks.Overlapping resources")
+IMPLEMENT_AI_INSTANT_TEST(FAITest_GameplayTask_OverlappingResources, "System.Engine.AI.Gameplay Tasks.Overlapping resources")
+
+//----------------------------------------------------------------------//
+// Pausing a task overlapping a lower priority task should not resume the low priority task
+//----------------------------------------------------------------------//
+struct FAITest_GameplayTask_PausingTasksBlockingOtherTasks : public FAITest_GameplayTasksTest
+{
+	UMockTask_Log* Tasks[3];
+
+	void InstantTest()
+	{
+		UWorld& World = GetWorld();
+
+		Tasks[0] = UMockTask_Log::CreateTask(*Component, Logger, MovementResourceSet);
+		Tasks[1] = UMockTask_Log::CreateTask(*Component, Logger, MoveAndLogicResourceSet);
+		Tasks[2] = UMockTask_Log::CreateTask(*Component, Logger, LogicResourceSet);
+		 
+		Tasks[0]->ReadyForActivation();
+		Tasks[1]->ReadyForActivation();
+
+		Test(TEXT("First task should be paused since it's resources get overlapped"), Tasks[0]->IsActive() == false);
+		Test(TEXT("Second task should on top and active"), Tasks[1]->IsActive() == true);
+
+		Tasks[2]->ReadyForActivation();
+		Test(TEXT("Second task should get paused since its resources got overlapped"), Tasks[1]->IsActive() == false);
+		Test(TEXT("First task should remain paused since it's resources get overlapped by the paused task"), Tasks[0]->IsActive() == false);
+
+		Tasks[2]->ExternalCancel();
+		Test(TEXT("Nothing shoud change for the first task"), Tasks[0]->IsActive() == false);
+		Test(TEXT("Second task should be active again"), Tasks[1]->IsActive() == true);
+	}
+};
+IMPLEMENT_AI_INSTANT_TEST(FAITest_GameplayTask_PausingTasksBlockingOtherTasks, "System.Engine.AI.Gameplay Tasks.Pausing tasks blocking other tasks")
 
 //----------------------------------------------------------------------//
 // Priority handling
@@ -263,7 +295,7 @@ struct FAITest_GameplayTask_Priorities : public FAITest_GameplayTasksTest
 {
 	UMockTask_Log* Tasks[3];
 
-	void SetUp()
+	void InstantTest()
 	{
 		UWorld& World = GetWorld();
 
@@ -279,7 +311,7 @@ struct FAITest_GameplayTask_Priorities : public FAITest_GameplayTasksTest
 		Tasks[0]->ReadyForActivation();
 
 		Test(TEXT("Task at index 1 should be active at this point since it's higher priority"), Tasks[1]->IsActive() && !Tasks[0]->IsActive());
-		Test(TEXT("TasksComponent should claim it's using only resources of task 1"), Component->GetCurrentlyUsedResources() == Tasks[1]->GetRequiredResources());
+		Test(TEXT("TasksComponent should claim it's using only resources of task 1"), Component->GetCurrentlyUsedResources() == Tasks[1]->GetClaimedResources());
 		Test(TEXT("Current top action wants to tick so Component should want that as well"), Component->GetShouldTick() == true); 
 
 		Tasks[2]->ReadyForActivation();
@@ -298,7 +330,7 @@ struct FAITest_GameplayTask_Priorities : public FAITest_GameplayTasksTest
 		Test(TEXT("Task-less component should not want to tick"), Component->GetShouldTick() == false);
 	}
 };
-IMPLEMENT_AI_LATENT_TEST(FAITest_GameplayTask_Priorities, "System.Engine.AI.Gameplay Tasks.Priorities")
+IMPLEMENT_AI_INSTANT_TEST(FAITest_GameplayTask_Priorities, "System.Engine.AI.Gameplay Tasks.Priorities")
 
 //----------------------------------------------------------------------//
 // Internal ending, by task ending itself or owner finishing 
@@ -312,7 +344,7 @@ struct FAITest_GameplayTask_InternalEnding : public FAITest_GameplayTasksTest
 	UMockTask_Log* Task_OwnerEndWithResources;
 	UMockTask_Log* Tasks[TasksCount];
 
-	void SetUp()
+	void InstantTest()
 	{
 		UWorld& World = GetWorld();
 
@@ -347,8 +379,96 @@ struct FAITest_GameplayTask_InternalEnding : public FAITest_GameplayTasksTest
 		Test(TEXT("Priority Task Queue should be empty"), Component->GetTaskPriorityQueueSize() == 0);
 	}
 };
-IMPLEMENT_AI_LATENT_TEST(FAITest_GameplayTask_InternalEnding, "System.Engine.AI.Gameplay Tasks.Self and Owner ending")
+IMPLEMENT_AI_INSTANT_TEST(FAITest_GameplayTask_InternalEnding, "System.Engine.AI.Gameplay Tasks.Self and Owner ending")
+
+//----------------------------------------------------------------------//
+// 
+//----------------------------------------------------------------------//
+struct FAITest_GameplayTask_MultipleOwners : public FAITest_GameplayTasksTest
+{
+	static const int32 TasksCount = 3;
+	UMockTask_Log* Tasks[TasksCount];
+	UMockTask_Log* LowPriorityTask;
+	UMockGameplayTaskOwner* OtherOwner;
+
+	void InstantTest()
+	{
+		OtherOwner = NewObject<UMockGameplayTaskOwner>();
+		OtherOwner->GTComponent = Component;
+		Tasks[0] = UMockTask_Log::CreateTask(*Component, Logger, MovementResourceSet);
+		Tasks[1] = UMockTask_Log::CreateTask(*OtherOwner, Logger, MovementResourceSet);
+		Tasks[2] = UMockTask_Log::CreateTask(*Component, Logger, MovementResourceSet);
+
+		for (int32 Index = 0; Index < TasksCount; ++Index)
+		{
+			Tasks[Index]->ReadyForActivation();
+		}
+		// This part tests what happens if "other owner" task is in the middle of the queue and
+		// not active
+		Test(TEXT("Last pushed task should be active now"), Tasks[2]->IsActive() == true);
+		Component->EndAllResourceConsumingTasksOwnedBy(*Component);
+		Test(TEXT("There should be only one task in the queue now"), Component->GetTaskPriorityQueueSize() == 1);
+		Test(TEXT("The last remaining task should be active now"), Tasks[1]->IsActive() == true);
+
+		// this part tests what happens during pruning if the "other owner" task is active at the
+		// moment of performing the action
+		LowPriorityTask = UMockTask_Log::CreateTask(*Component, Logger, MovementResourceSet, LowPriority);
+		LowPriorityTask->ReadyForActivation();
+		Test(TEXT("There should be 2 tasks in the queue now"), Component->GetTaskPriorityQueueSize() == 2);
+		Component->EndAllResourceConsumingTasksOwnedBy(*Component);
+		Test(TEXT("There should be only one task in the queue after second pruning"), Component->GetTaskPriorityQueueSize() == 1);
+		Test(TEXT("The last remaining task should be still active"), Tasks[1]->IsActive() == true);
+	}
+};
+IMPLEMENT_AI_INSTANT_TEST(FAITest_GameplayTask_MultipleOwners, "System.Engine.AI.Gameplay Tasks.Handling multiple task owners")
+
+//----------------------------------------------------------------------//
+// Claimed vs Required resources test
+//----------------------------------------------------------------------//
+struct FAITest_GameplayTask_ClaimedResources : public FAITest_GameplayTasksTest
+{
+	static const int32 TasksCount = 4;
+	UMockTask_Log* Tasks[TasksCount];
+
+	void InstantTest()
+	{
+		UWorld& World = GetWorld();
+
+		// create three tasks
+		// first one has a resource we're going to overlap with the extra-claimed resource of the next task
+		Tasks[0] = UMockTask_Log::CreateTask(*Component, Logger, MovementResourceSet);
+		Tasks[0]->ReadyForActivation();
+
+		// second task requires non-overlapping resources to first task
+		Tasks[1] = UMockTask_Log::CreateTask(*Component, Logger, AnimationResourceSet);
+		// but declared an overlapping resource as "claimed"
+		Tasks[1]->AddClaimedResourceSet(MovementResourceSet);
+		
+		Tasks[1]->ReadyForActivation();
+		// at this point first task should get paused since it's required resource is claimed, or shadowed, by the newer task
+		Test(TEXT("The first task should get paused since its required resource is claimed, or shadowed, by the newer task"), Tasks[0]->IsActive() == false);
+		Test(TEXT("The second task should be running, nothing obstructing it"), Tasks[1]->IsActive() == true);
+		
+		// a new low-priority task should not be allowed to run neither
+		Tasks[2] = UMockTask_Log::CreateTask(*Component, Logger, MovementResourceSet, LowPriority);
+		Tasks[2]->ReadyForActivation();
+		Test(TEXT("The new low-priority task should not be allowed to run neither"), Tasks[2]->IsActive() == false);
+		Test(TEXT("The second task should be still running"), Tasks[1]->IsActive() == true);
+
+		// however, a new task, that's using the overlapped claimed resource 
+		// should run without any issues
+		// note, this doesn't have to be "high priority" task - new tasks with same priority as "current"
+		// are treated like higher priority anyway
+		Tasks[3] = UMockTask_Log::CreateTask(*Component, Logger, MovementResourceSet, HighPriority);
+		Tasks[3]->ReadyForActivation();
+		Test(TEXT("The new high-priority task should be allowed to run"), Tasks[3]->IsActive() == true);
+		// but active task, that declared the active resources should not get paused neither
+		Test(TEXT("The second task should be still running, it's required resources are not being overlapped"), Tasks[1]->IsActive() == true);
+	}
+};
+IMPLEMENT_AI_INSTANT_TEST(FAITest_GameplayTask_ClaimedResources, "System.Engine.AI.Gameplay Tasks.Claimed resources")
 
 // add tests if component wants ticking at while aborting/reactivating tasks
+// add test for re-adding/re-activating a finished task
 
 #undef LOCTEXT_NAMESPACE
\ No newline at end of file
diff --git a/Engine/Source/Runtime/AIModule/Classes/AIController.h b/Engine/Source/Runtime/AIModule/Classes/AIController.h
--- a/Engine/Source/Runtime/AIModule/Classes/AIController.h
+++ b/Engine/Source/Runtime/AIModule/Classes/AIController.h
@@ -3,6 +3,7 @@
 #pragma once
 
 #include "AITypes.h"
+#include "GameplayTask.h"
 #include "GameFramework/Pawn.h"
 #include "AI/Navigation/NavigationTypes.h"
 #include "AI/Navigation/NavigationSystem.h"
@@ -18,6 +19,7 @@
 class APawn;
 class UPathFollowingComponent;
 class UBrainComponent;
+class UBlackboardComponent;
 class UAIPerceptionComponent;
 class UPawnAction;
 class UPawnActionsComponent;
@@ -122,6 +124,10 @@ private_subobject:
 	UPawnActionsComponent* ActionsComp;
 
 protected:
+	/** blackboard */
+	UPROPERTY(BlueprintReadOnly, Category = AI, meta = (AllowPrivateAccess = "true"))
+	UBlackboardComponent* Blackboard;
+
 	UPROPERTY()
 	UGameplayTasksComponent* CachedGameplayTasksComponent;
 
@@ -133,6 +139,8 @@ public:
 	UFUNCTION(BlueprintImplementableEvent)
 	void OnPossess(APawn* PossessedPawn);
 
+	virtual void SetPawn(APawn* InPawn) override;
+
 	/** Makes AI go toward specified Goal actor (destination will be continuously updated)
 	 *  @param AcceptanceRadius - finish move if pawn gets close enough
 	 *  @param bStopOnOverlap - add pawn's radius to AcceptanceRadius
@@ -358,13 +366,18 @@ public:
 	//----------------------------------------------------------------------//
 	// IGameplayTaskOwnerInterface
 	//----------------------------------------------------------------------//
-	virtual UGameplayTasksComponent* GetGameplayTasksComponent() override { return CachedGameplayTasksComponent; }
+	UGameplayTasksComponent* GetGameplayTasksComponent() const { return CachedGameplayTasksComponent; }
+	virtual UGameplayTasksComponent* GetGameplayTasksComponent(const UGameplayTask& Task) const override { return CachedGameplayTasksComponent; }
 	virtual void OnTaskActivated(UGameplayTask& Task) override {}
 	virtual void OnTaskDeactivated(UGameplayTask& Task) override {}
 	virtual void OnTaskInitialized(UGameplayTask& Task) override {}
-	virtual AActor* GetOwnerActor() const { return const_cast<AAIController*>(this); }
-	virtual AActor* GetAvatarActor() const { return GetPawn(); }
-	virtual uint8 GetDefaultPriority() const { return FGameplayTasks::DefaultPriority - 1; }
+	virtual AActor* GetOwnerActor(const UGameplayTask* Task) const override { return const_cast<AAIController*>(this); }
+	virtual AActor* GetAvatarActor(const UGameplayTask* Task) const override { return GetPawn(); }
+	virtual uint8 GetDefaultPriority() const override { return FGameplayTasks::DefaultPriority - 1; }
+
+	// other GT tasks related
+	UFUNCTION()
+	virtual void OnGameplayTaskResourcesClaimed(FGameplayResourceSet NewlyClaimed, FGameplayResourceSet FreshlyReleased);
 
 	//----------------------------------------------------------------------//
 	// Actions
diff --git a/Engine/Source/Runtime/AIModule/Classes/AISystem.h b/Engine/Source/Runtime/AIModule/Classes/AISystem.h
--- a/Engine/Source/Runtime/AIModule/Classes/AISystem.h
+++ b/Engine/Source/Runtime/AIModule/Classes/AISystem.h
@@ -15,17 +15,37 @@ class UAIHotSpotManager;
 class UBlackboardData;
 class UBlackboardComponent;
 
+#define GET_AI_CONFIG_VAR(a) (GetDefault<UAISystem>()->a)
+
 UCLASS(config=Engine, defaultconfig)
 class AIMODULE_API UAISystem : public UAISystemBase
 {
-	GENERATED_UCLASS_BODY()
+	GENERATED_BODY()
 
+protected:
 	UPROPERTY(globalconfig, EditAnywhere, Category = "AISystem", meta = (MetaClass = "AIPerceptionSystem", DisplayName = "Perception System Class"))
 	FStringClassReference PerceptionSystemClassName;
 
 	UPROPERTY(globalconfig, EditAnywhere, Category = "AISystem", meta = (MetaClass = "AIHotSpotManager", DisplayName = "AIHotSpotManager Class"))
 	FStringClassReference HotSpotManagerClassName;
 
+public:
+	UPROPERTY(globalconfig, EditDefaultsOnly, Category = "AISystem | Movement")
+	float AcceptanceRadius; 
+	
+	UPROPERTY(globalconfig, EditDefaultsOnly, Category = "AISystem | Movement")
+	bool bFinishMoveOnGoalOverlap;
+
+	UPROPERTY(globalconfig, EditDefaultsOnly, Category = "AISystem | Movement")
+	bool bAcceptPartialPaths;
+
+	UPROPERTY(globalconfig, EditDefaultsOnly, Category = "AISystem | Movement")
+	bool bAllowStrafing;
+
+	/** this property is just a transition-time flag - in the end we're going to switch over to Gameplay Tasks anyway, that's the goal. */
+	UPROPERTY(globalconfig, EditDefaultsOnly, Category = "AISystem | Gameplay Tasks")
+	bool bEnableBTAITasks;
+
 protected:
 	/** Behavior tree manager used by game */
 	UPROPERTY(Transient)
@@ -52,6 +72,8 @@ protected:
 	FDelegateHandle ActorSpawnedDelegateHandle;
 	
 public:
+	UAISystem(const FObjectInitializer& ObjectInitializer);
+
 	virtual void BeginDestroy() override;
 	
 	virtual void PostInitProperties() override;
diff --git a/Engine/Source/Runtime/AIModule/Classes/AITypes.h b/Engine/Source/Runtime/AIModule/Classes/AITypes.h
--- a/Engine/Source/Runtime/AIModule/Classes/AITypes.h
+++ b/Engine/Source/Runtime/AIModule/Classes/AITypes.h
@@ -53,6 +53,19 @@ namespace EAIOptionFlag
 	};
 }
 
+namespace FAISystem
+{
+	FORCEINLINE bool PickAIOption(EAIOptionFlag::Type Option, bool DefaultOption)
+	{
+		return Option == EAIOptionFlag::Default ? DefaultOption : (Option == EAIOptionFlag::Enable);
+	}
+
+	FORCEINLINE EAIOptionFlag::Type BoolToAIOption(bool Value)
+	{
+		return Value ? EAIOptionFlag::Enable : EAIOptionFlag::Disable;
+	}
+}
+
 namespace EAIForceParam
 {
 	enum Type
diff --git a/Engine/Source/Runtime/AIModule/Classes/BehaviorTree/BTTaskNode.h b/Engine/Source/Runtime/AIModule/Classes/BehaviorTree/BTTaskNode.h
--- a/Engine/Source/Runtime/AIModule/Classes/BehaviorTree/BTTaskNode.h
+++ b/Engine/Source/Runtime/AIModule/Classes/BehaviorTree/BTTaskNode.h
@@ -3,6 +3,7 @@
 #pragma once
 
 #include "BehaviorTree/BTNode.h"
+#include "GameplayTaskOwnerInterface.h"
 #include "BTTaskNode.generated.h"
 
 class UBehaviorTreeComponent;
@@ -24,7 +25,7 @@ struct FAIMessage;
  */
 
 UCLASS(Abstract)
-class AIMODULE_API UBTTaskNode : public UBTNode
+class AIMODULE_API UBTTaskNode : public UBTNode, public IGameplayTaskOwnerInterface
 {
 	GENERATED_UCLASS_BODY()
 
@@ -33,11 +34,13 @@ class AIMODULE_API UBTTaskNode : public UBTNode
 	 * this function should be considered as const (don't modify state of object) if node is not instanced! */
 	virtual EBTNodeResult::Type ExecuteTask(UBehaviorTreeComponent& OwnerComp, uint8* NodeMemory);
 
+protected:
 	/** aborts this task, should return Aborted or InProgress
 	 *  (use FinishLatentAbort() when returning InProgress)
 	 * this function should be considered as const (don't modify state of object) if node is not instanced! */
 	virtual EBTNodeResult::Type AbortTask(UBehaviorTreeComponent& OwnerComp, uint8* NodeMemory);
 
+public:
 #if WITH_EDITOR
 	virtual FName GetNodeIconName() const override;
 #endif // WITH_EDITOR
@@ -71,6 +74,10 @@ protected:
 	/** if set, OnTaskFinished will be called */
 	uint8 bNotifyTaskFinished : 1;
 
+	/** set to true if task owns any GameplayTasks. Note this requires tasks to be created via NewBTAITask
+	 *	Otherwise specific BT task node class is responsible for ending the gameplay tasks on node finish */
+	uint8 bOwnsGameplayTasks : 1;
+
 	/** ticks this task 
 	 * this function should be considered as const (don't modify state of object) if node is not instanced! */
 	virtual void TickTask(UBehaviorTreeComponent& OwnerComp, uint8* NodeMemory, float DeltaSeconds);
@@ -90,6 +97,31 @@ protected:
 	/** unregister message observers */
 	void StopWaitingForMessages(UBehaviorTreeComponent& OwnerComp) const;
 
+	template <class T>
+	T* NewBTAITask(UBehaviorTreeComponent& BTComponent)
+	{
+		T* NewAITask = NewObject<T>();
+		AAIController* AIController = BTComponent.GetAIOwner();
+		check(AIController && "Can\'t spawn an AI task without AI controller!");
+		NewAITask->InitAITask(*AIController, *this, GetDefaultPriority());
+		bOwnsGameplayTasks = true;
+		
+		return NewAITask;
+	}
+
+	//----------------------------------------------------------------------//
+	// IGameplayTaskOwnerInterface
+	//----------------------------------------------------------------------//
+	virtual UGameplayTasksComponent* GetGameplayTasksComponent(const UGameplayTask& Task) const override;
+	virtual void OnTaskActivated(UGameplayTask& Task) override;
+	virtual void OnTaskDeactivated(UGameplayTask& Task) override;
+	virtual void OnTaskInitialized(UGameplayTask& Task) override;
+	virtual AActor* GetOwnerActor(const UGameplayTask* Task) const;
+	virtual AActor* GetAvatarActor(const UGameplayTask* Task) const;
+	virtual uint8 GetDefaultPriority() const;
+
+	UBehaviorTreeComponent* GetBTComponentForTask(UGameplayTask& Task) const;
+
 	//----------------------------------------------------------------------//
 	// DEPRECATED
 	//----------------------------------------------------------------------//
diff --git a/Engine/Source/Runtime/AIModule/Classes/BehaviorTree/Tasks/BTTask_MoveDirectlyToward.h b/Engine/Source/Runtime/AIModule/Classes/BehaviorTree/Tasks/BTTask_MoveDirectlyToward.h
--- a/Engine/Source/Runtime/AIModule/Classes/BehaviorTree/Tasks/BTTask_MoveDirectlyToward.h
+++ b/Engine/Source/Runtime/AIModule/Classes/BehaviorTree/Tasks/BTTask_MoveDirectlyToward.h
@@ -36,7 +36,7 @@ class AIMODULE_API UBTTask_MoveDirectlyToward : public UBTTask_BlackboardBase
 
 	/** if set to true agent's radius will be added to AcceptableRadius for purposes of checking
 	 *	if path's end point has been reached. Will result in AI stopping on contact with destination location*/
-	UPROPERTY(Category = Node, EditAnywhere, config)
+	UPROPERTY(Category = Node, EditAnywhere)
 	uint32 bStopOnOverlap : 1;
 
 	virtual EBTNodeResult::Type ExecuteTask(UBehaviorTreeComponent& OwnerComp, uint8* NodeMemory) override;
diff --git a/Engine/Source/Runtime/AIModule/Classes/BehaviorTree/Tasks/BTTask_MoveTo.h b/Engine/Source/Runtime/AIModule/Classes/BehaviorTree/Tasks/BTTask_MoveTo.h
--- a/Engine/Source/Runtime/AIModule/Classes/BehaviorTree/Tasks/BTTask_MoveTo.h
+++ b/Engine/Source/Runtime/AIModule/Classes/BehaviorTree/Tasks/BTTask_MoveTo.h
@@ -38,7 +38,7 @@ class AIMODULE_API UBTTask_MoveTo : public UBTTask_BlackboardBase
 
 	/** if set to true agent's radius will be added to AcceptableRadius for purposes of checking 
 	 *	if path's end point has been reached. Will result in AI stopping on contact with destination location*/
-	UPROPERTY(Category=Node, EditAnywhere, config)
+	UPROPERTY(Category=Node, EditAnywhere)
 	uint32 bStopOnOverlap : 1;
 	
 	virtual EBTNodeResult::Type ExecuteTask(UBehaviorTreeComponent& OwnerComp, uint8* NodeMemory) override;
diff --git a/Engine/Source/Runtime/AIModule/Classes/Tasks/AITask.h b/Engine/Source/Runtime/AIModule/Classes/Tasks/AITask.h
new file mode 100644
--- /dev/null
+++ b/Engine/Source/Runtime/AIModule/Classes/Tasks/AITask.h
@@ -0,0 +1,41 @@
+// Copyright 1998-2015 Epic Games, Inc. All Rights Reserved.
+#pragma once
+
+#include "GameplayTask.h"
+#include "AITypes.h"
+#include "AITask.generated.h"
+
+class AAIController;
+
+UENUM()
+enum class EAITaskPriority : uint8
+{
+	Lowest = 0,
+	Low = 64, //FGameplayTasks::DefaultPriority / 2,
+	AutonomousAI = 127, //FGameplayTasks::DefaultPriority,
+	High = 192, //(1.5 * FGameplayTasks::DefaultPriority),
+	Ultimate = 255,
+};
+
+UCLASS(Abstract)
+class AIMODULE_API UAITask : public UGameplayTask
+{
+	GENERATED_BODY()
+protected:
+
+	UPROPERTY()
+	AAIController* OwnerController;
+
+	virtual void Activate() override;
+
+public:
+	UAITask(const FObjectInitializer& ObjectInitializer);
+
+	static AAIController* GetAIControllerForActor(AActor* Actor);
+	AAIController* GetAIController() const { return OwnerController; };
+
+	void InitAITask(AAIController& AIOwner, IGameplayTaskOwnerInterface& InTaskOwner, uint8 InPriority);
+
+	/** effectively adds UAIResource_Logic to the set of Claimed resources */
+	void RequestAILogicLocking();
+};
\ No newline at end of file
diff --git a/Engine/Source/Runtime/AIModule/Classes/Tasks/AITask_MoveTo.h b/Engine/Source/Runtime/AIModule/Classes/Tasks/AITask_MoveTo.h
--- a/Engine/Source/Runtime/AIModule/Classes/Tasks/AITask_MoveTo.h
+++ b/Engine/Source/Runtime/AIModule/Classes/Tasks/AITask_MoveTo.h
@@ -1,18 +1,16 @@
 // Copyright 1998-2015 Epic Games, Inc. All Rights Reserved.
 #pragma once
 
-#include "GameplayTask.h"
-#include "AITypes.h"
+#include "AITask.h"
 #include "Navigation/PathFollowingComponent.h"
 #include "AITask_MoveTo.generated.h"
 
 class APawn;
-class AAIController;
 
-DECLARE_DYNAMIC_MULTICAST_DELEGATE_OneParam(FMoveTaskCompletedSignature, const TEnumAsByte<EPathFollowingResult::Type>&, Result);
+DECLARE_DYNAMIC_MULTICAST_DELEGATE_OneParam(FMoveTaskCompletedSignature, TEnumAsByte<EPathFollowingResult::Type>, Result);
 
 UCLASS()
-class AIMODULE_API UAITask_MoveTo : public UGameplayTask
+class AIMODULE_API UAITask_MoveTo : public UAITask
 {
 	GENERATED_BODY()
 protected:
@@ -22,27 +20,46 @@ protected:
 	UPROPERTY(BlueprintAssignable)
 	FMoveTaskCompletedSignature OnMoveFinished;
 
-	UPROPERTY()
-	AAIController* OwnerController;
+	UPROPERTY(BlueprintReadOnly, meta = (ExposeOnSpawn = true, DisplayName="Goal Location"), Category = "AITask")
+	FVector MoveGoalLocation;
+
+	/** gets set to GoalLocation if GoalActor != nullptr, otherwise is FAISystem::InvalidLocation */
+	FVector RealGoalLocation;
+	
+	UPROPERTY(BlueprintReadOnly, meta = (ExposeOnSpawn = true, DisplayName="Goal Actor"), Category = "AITask")
+	AActor* MoveGoalActor;
+
+	UPROPERTY(BlueprintReadOnly, meta = (ExposeOnSpawn = true, DisplayName = "Acceptance Radius"), Category = "AITask", AdvancedDisplay)
+	float MoveAcceptanceRadius; 
 	
-	UPROPERTY()
-	FAIMoveRequest MoveRequest;
+	UPROPERTY(BlueprintReadOnly, meta = (ExposeOnSpawn = true), Category = "AITask", AdvancedDisplay)
+	bool bShouldStopOnOverlap;
+
+	UPROPERTY(BlueprintReadOnly, meta = (ExposeOnSpawn = true), Category = "AITask", AdvancedDisplay)
+	bool bShouldAcceptPartialPath;
+
+	UPROPERTY(BlueprintReadOnly, meta = (ExposeOnSpawn = true), Category = "AITask", AdvancedDisplay)
+	bool bShouldUsePathfinding;
 
 	FAIRequestID MoveRequestID;
 
 	virtual void HandleMoveFinished(FAIRequestID RequestID, EPathFollowingResult::Type Result);
 	virtual void Activate() override;
+	virtual void OnDestroy(bool bOwnerFinished) override;
 
 	virtual void Pause() override;
 	virtual void Resume() override;
 
+	virtual void PostInitProperties() override;
+
 public:
 	UAITask_MoveTo(const FObjectInitializer& ObjectInitializer);
 
-	/** Spawn new Actor on the network authority (server) */
-	UFUNCTION(BlueprintCallable, Category = "GameplayTasks", meta = (AdvancedDisplay = "TaskOwner,AcceptanceRadius,bStopOnOverlap,bAcceptPartialPath", DefaultToSelf = "TaskOwner", BlueprintInternalUseOnly = "TRUE", DisplayName="Move To Location or Actor"))
-	static UAITask_MoveTo* PawnMoveTo(APawn* Pawn, FVector Destination, AActor* TargetActor = NULL, float AcceptanceRadius = 5.f, bool bStopOnOverlap = false, bool bAcceptPartialPath = false);
+	/**
+	 *	@param 
+	 */
+	UFUNCTION(BlueprintCallable, Category = "AI|Tasks", meta = (AdvancedDisplay = "TaskOwner,AcceptanceRadius,StopOnOverlap,AcceptPartialPath,bUsePathfinding", DefaultToSelf = "TaskOwner", BlueprintInternalUseOnly = "TRUE", DisplayName = "Move To Location or Actor"))
+	static UAITask_MoveTo* AIMoveTo(AAIController* Controller, FVector GoalLocation, AActor* GoalActor = nullptr, float AcceptanceRadius = -1.f, EAIOptionFlag::Type StopOnOverlap = EAIOptionFlag::Default, EAIOptionFlag::Type AcceptPartialPath = EAIOptionFlag::Default, bool bUsePathfinding = true, bool bLockAILogic = true);
 
-	UFUNCTION(BlueprintCallable, Category = "AI|Tasks", meta = (AdvancedDisplay = "TaskOwner,AcceptanceRadius,bStopOnOverlap,bAcceptPartialPath", DefaultToSelf = "TaskOwner", BlueprintInternalUseOnly = "TRUE", DisplayName = "Move To Location or Actor"))
-	static UAITask_MoveTo* AIMoveTo(AAIController* Controller, FVector TargetLocation, AActor* TargetActor = NULL, float AcceptanceRadius = 5.f, bool bStopOnOverlap = false, bool bAcceptPartialPath = false);
+	void SetUp(AAIController* Controller, FVector GoalLocation, AActor* GoalActor = nullptr, float AcceptanceRadius = -1.f, bool bUsePathfinding = true, EAIOptionFlag::Type StopOnOverlap = EAIOptionFlag::Default, EAIOptionFlag::Type AcceptPartialPath = EAIOptionFlag::Default);
 };
\ No newline at end of file
diff --git a/Engine/Source/Runtime/AIModule/Private/AIController.cpp b/Engine/Source/Runtime/AIModule/Private/AIController.cpp
--- a/Engine/Source/Runtime/AIModule/Private/AIController.cpp
+++ b/Engine/Source/Runtime/AIModule/Private/AIController.cpp
@@ -478,6 +478,8 @@ void AAIController::Possess(APawn* InPawn)
 
 	if (GTComp)
 	{
+		GTComp->OnClaimedResourcesChange.AddDynamic(this, &AAIController::OnGameplayTaskResourcesClaimed);
+
 		REDIRECT_OBJECT_TO_VLOG(GTComp, this);
 	}
 
@@ -501,6 +503,24 @@ void AAIController::UnPossess()
 	CachedGameplayTasksComponent = nullptr;
 }
 
+void AAIController::SetPawn(APawn* InPawn)
+{
+	Super::SetPawn(InPawn);
+
+	if (Blackboard)
+	{
+		const UBlackboardData* BBAsset = Blackboard->GetBlackboardAsset();
+		if (BBAsset)
+		{
+			const FBlackboard::FKey SelfKey = BBAsset->GetKeyID(FBlackboard::KeySelf);
+			if (SelfKey != FBlackboard::InvalidKey)
+			{
+				Blackboard->SetValue<UBlackboardKeyType_Object>(SelfKey, GetPawn());
+			}
+		}
+	}
+}
+
 void AAIController::InitNavigationControl(UPathFollowingComponent*& PathFollowingComp)
 {
 	PathFollowingComp = PathFollowingComponent;
@@ -850,8 +870,16 @@ bool AAIController::RunBehaviorTree(UBehaviorTree* BTAsset)
 bool AAIController::InitializeBlackboard(UBlackboardComponent& BlackboardComp, UBlackboardData& BlackboardAsset)
 {
 	check(BlackboardComp.GetOwner() == this);
+
 	if (BlackboardComp.InitializeBlackboard(BlackboardAsset))
 	{
+		// find the "self" key and set it to our pawn
+		const FBlackboard::FKey SelfKey = BlackboardAsset.GetKeyID(FBlackboard::KeySelf);
+		if (SelfKey != FBlackboard::InvalidKey)
+		{
+			BlackboardComp.SetValue<UBlackboardKeyType_Object>(SelfKey, GetPawn());
+		}
+
 		OnUsingBlackBoard(&BlackboardComp, &BlackboardAsset);
 		return true;
 	}
@@ -860,37 +888,38 @@ bool AAIController::InitializeBlackboard(UBlackboardComponent& BlackboardComp, U
 
 bool AAIController::UseBlackboard(UBlackboardData* BlackboardAsset, UBlackboardComponent*& BlackboardComponent)
 {
-	if (BlackboardAsset == NULL)
+	if (BlackboardAsset == nullptr)
 	{
 		UE_VLOG(this, LogBehaviorTree, Log, TEXT("UseBlackboard: trying to use NULL Blackboard asset. Ignoring"));
 		return false;
 	}
 
 	bool bSuccess = true;
-	UBlackboardComponent* BlackboardComp = FindComponentByClass<UBlackboardComponent>();
+	Blackboard = FindComponentByClass<UBlackboardComponent>();
 
-	if (BlackboardComp == NULL)
+	if (Blackboard == nullptr)
 	{
-		BlackboardComp = NewObject<UBlackboardComponent>(this, TEXT("BlackboardComponent"));
-		if (BlackboardComp != NULL)
+		Blackboard = NewObject<UBlackboardComponent>(this, TEXT("BlackboardComponent"));
+		if (Blackboard != nullptr)
 		{
-			InitializeBlackboard(*BlackboardComp, *BlackboardAsset);
-			BlackboardComp->RegisterComponent();
+			InitializeBlackboard(*Blackboard, *BlackboardAsset);
+			Blackboard->RegisterComponent();
 		}
 
 	}
-	else if (BlackboardComp->GetBlackboardAsset() == NULL)
+	else if (Blackboard->GetBlackboardAsset() == nullptr)
 	{
-		InitializeBlackboard(*BlackboardComp, *BlackboardAsset);
+		InitializeBlackboard(*Blackboard, *BlackboardAsset);
 	}
-	else if (BlackboardComp->GetBlackboardAsset() != BlackboardAsset)
+	else if (Blackboard->GetBlackboardAsset() != BlackboardAsset)
 	{
+		// @todo this behavior should be opt-out-able.
 		UE_VLOG(this, LogBehaviorTree, Log, TEXT("UseBlackboard: requested blackboard %s while already has %s instantiated. Forcing new BB.")
-			, *GetNameSafe(BlackboardAsset), *GetNameSafe(BlackboardComp->GetBlackboardAsset()));
-		InitializeBlackboard(*BlackboardComp, *BlackboardAsset);
+			, *GetNameSafe(BlackboardAsset), *GetNameSafe(Blackboard->GetBlackboardAsset()));
+		InitializeBlackboard(*Blackboard, *BlackboardAsset);
 	}
 
-	BlackboardComponent = BlackboardComp;
+	BlackboardComponent = Blackboard;
 
 	return bSuccess;
 }
@@ -935,6 +964,22 @@ const UAIPerceptionComponent* AAIController::GetAIPerceptionComponent() const
 	return PerceptionComponent;
 }
 
+void AAIController::OnGameplayTaskResourcesClaimed(FGameplayResourceSet NewlyClaimed, FGameplayResourceSet FreshlyReleased)
+{
+	if (BrainComponent)
+	{
+		const uint8 LogicID = UGameplayTaskResource::GetResourceID<UAIResource_Logic>();
+		if (NewlyClaimed.HasID(LogicID))
+		{
+			BrainComponent->LockResource(EAIRequestPriority::Logic);
+		}
+		else if (FreshlyReleased.HasID(LogicID))
+		{
+			BrainComponent->ClearResourceLock(EAIRequestPriority::Logic);
+		}
+	}
+}
+
 //----------------------------------------------------------------------//
 // IGenericTeamAgentInterface
 //----------------------------------------------------------------------//
diff --git a/Engine/Source/Runtime/AIModule/Private/AIModulePrivate.h b/Engine/Source/Runtime/AIModule/Private/AIModulePrivate.h
--- a/Engine/Source/Runtime/AIModule/Private/AIModulePrivate.h
+++ b/Engine/Source/Runtime/AIModule/Private/AIModulePrivate.h
@@ -71,7 +71,8 @@
 #include "BehaviorTree/BTFunctionLibrary.h"
 
 #include "AIResources.h"
+#include "Tasks/AITask.h"
 
 DECLARE_LOG_CATEGORY_EXTERN(LogAINavigation, Warning, All);
 DECLARE_LOG_CATEGORY_EXTERN(LogBehaviorTree, Display, All);
-DECLARE_LOG_CATEGORY_EXTERN(LogCrowdFollowing, Warning, All);
\ No newline at end of file
+DECLARE_LOG_CATEGORY_EXTERN(LogCrowdFollowing, Warning, All);
diff --git a/Engine/Source/Runtime/AIModule/Private/AISystem.cpp b/Engine/Source/Runtime/AIModule/Private/AISystem.cpp
--- a/Engine/Source/Runtime/AIModule/Private/AISystem.cpp
+++ b/Engine/Source/Runtime/AIModule/Private/AISystem.cpp
@@ -13,8 +13,17 @@
 #include "AISystem.h"
 
 
-UAISystem::UAISystem(const FObjectInitializer& ObjectInitializer) : Super(ObjectInitializer)
+UAISystem::UAISystem(const FObjectInitializer& ObjectInitializer)
+	: Super(ObjectInitializer)
 {
+	// default values of AI config params
+	AcceptanceRadius = 5.f;
+	bFinishMoveOnGoalOverlap = true;
+	bAcceptPartialPaths = true;
+	bAllowStrafing = false;
+
+	bEnableBTAITasks = false;
+
 	if (HasAnyFlags(RF_ClassDefaultObject) == false)
 	{
 		UWorld* WorldOuter = Cast<UWorld>(GetOuter());
diff --git a/Engine/Source/Runtime/AIModule/Private/BehaviorTree/BTTaskNode.cpp b/Engine/Source/Runtime/AIModule/Private/BehaviorTree/BTTaskNode.cpp
--- a/Engine/Source/Runtime/AIModule/Private/BehaviorTree/BTTaskNode.cpp
+++ b/Engine/Source/Runtime/AIModule/Private/BehaviorTree/BTTaskNode.cpp
@@ -1,6 +1,7 @@
 // Copyright 1998-2015 Epic Games, Inc. All Rights Reserved.
 
 #include "AIModulePrivate.h"
+#include "GameplayTasksComponent.h"
 #include "BehaviorTree/BTTaskNode.h"
 
 UBTTaskNode::UBTTaskNode(const FObjectInitializer& ObjectInitializer) : Super(ObjectInitializer)
@@ -28,8 +29,20 @@ EBTNodeResult::Type UBTTaskNode::WrappedExecuteTask(UBehaviorTreeComponent& Owne
 
 EBTNodeResult::Type UBTTaskNode::WrappedAbortTask(UBehaviorTreeComponent& OwnerComp, uint8* NodeMemory) const
 {
-	const UBTNode* NodeOb = bCreateNodeInstance ? GetNodeInstance(OwnerComp, NodeMemory) : this;
-	return NodeOb ? ((UBTTaskNode*)NodeOb)->AbortTask(OwnerComp, NodeMemory) : EBTNodeResult::Aborted;
+	UBTNode* NodeOb = const_cast<UBTNode*>(bCreateNodeInstance ? GetNodeInstance(OwnerComp, NodeMemory) : this);
+	UBTTaskNode* TaskNodeOb = static_cast<UBTTaskNode*>(NodeOb);
+	EBTNodeResult::Type Result = TaskNodeOb ? TaskNodeOb->AbortTask(OwnerComp, NodeMemory) : EBTNodeResult::Aborted;
+
+	if (TaskNodeOb && TaskNodeOb->bOwnsGameplayTasks && OwnerComp.GetAIOwner())
+	{
+		UGameplayTasksComponent* GTComp = OwnerComp.GetAIOwner()->GetGameplayTasksComponent();
+		if (GTComp)
+		{
+			GTComp->EndAllResourceConsumingTasksOwnedBy(*TaskNodeOb);
+		}
+	}
+
+	return Result;
 }
 
 void UBTTaskNode::WrappedTickTask(UBehaviorTreeComponent& OwnerComp, uint8* NodeMemory, float DeltaSeconds) const
@@ -46,12 +59,23 @@ void UBTTaskNode::WrappedTickTask(UBehaviorTreeComponent& OwnerComp, uint8* Node
 
 void UBTTaskNode::WrappedOnTaskFinished(UBehaviorTreeComponent& OwnerComp, uint8* NodeMemory, EBTNodeResult::Type TaskResult) const
 {
-	if (bNotifyTaskFinished)
+	UBTNode* NodeOb = const_cast<UBTNode*>(bCreateNodeInstance ? GetNodeInstance(OwnerComp, NodeMemory) : this);
+
+	if (NodeOb)
 	{
-		const UBTNode* NodeOb = bCreateNodeInstance ? GetNodeInstance(OwnerComp, NodeMemory) : this;
-		if (NodeOb)
+		UBTTaskNode* TaskNodeOb = static_cast<UBTTaskNode*>(NodeOb);
+		if (TaskNodeOb->bNotifyTaskFinished)
+		{
+			TaskNodeOb->OnTaskFinished(OwnerComp, NodeMemory, TaskResult);
+		}
+
+		if (TaskNodeOb->bOwnsGameplayTasks && OwnerComp.GetAIOwner())
 		{
-			((UBTTaskNode*)NodeOb)->OnTaskFinished(OwnerComp, NodeMemory, TaskResult);
+			UGameplayTasksComponent* GTComp = OwnerComp.GetAIOwner()->GetGameplayTasksComponent();
+			if (GTComp)
+			{
+				GTComp->EndAllResourceConsumingTasksOwnedBy(*TaskNodeOb);
+			}
 		}
 	}
 }
@@ -138,6 +162,119 @@ FName UBTTaskNode::GetNodeIconName() const
 
 #endif	// WITH_EDITOR
 
+//----------------------------------------------------------------------//
+// UBTTaskNode IGameplayTaskOwnerInterface
+//----------------------------------------------------------------------//
+UGameplayTasksComponent* UBTTaskNode::GetGameplayTasksComponent(const UGameplayTask& Task) const 
+{
+	const UAITask* AsAITask = Cast<const UAITask>(&Task);
+	if (AsAITask)
+	{
+		return AsAITask->GetAIController() ? AsAITask->GetAIController()->GetGameplayTasksComponent(Task) : nullptr;
+	}
+
+	return Task.GetGameplayTasksComponent();
+}
+
+void UBTTaskNode::OnTaskInitialized(UGameplayTask& Task)
+{
+	// validate the task
+	UAITask* AsAITask = Cast<UAITask>(&Task);
+	if (AsAITask != nullptr && AsAITask->GetAIController() == nullptr)
+	{
+		// this means that the task has either been created without specifying 
+		// UAITAsk::OwnerController's value (like via BP's Construct Object node)
+		// or it has been created in C++ with inappropriate function
+		UE_LOG(LogBehaviorTree, Error, TEXT("Missing AIController in AITask %s"), *AsAITask->GetName());
+	}
+}
+
+UBehaviorTreeComponent* UBTTaskNode::GetBTComponentForTask(UGameplayTask& Task) const
+{
+	UAITask* AsAITask = Cast<UAITask>(&Task);
+	return AsAITask && AsAITask->GetAIController() ? Cast<UBehaviorTreeComponent>(AsAITask->GetAIController()->BrainComponent) : nullptr;
+}
+
+void UBTTaskNode::OnTaskActivated(UGameplayTask& Task) 
+{
+	ensure(Task.GetTaskOwner() == this);
+}
+
+void UBTTaskNode::OnTaskDeactivated(UGameplayTask& Task)
+{
+	ensure(Task.GetTaskOwner() == this);
+	UBehaviorTreeComponent* BTComp = GetBTComponentForTask(Task);
+	if (BTComp)
+	{
+		// this is a super-default behavior. Specific task will surely like to 
+		// handle this themselves, finishing with specific result
+		FinishLatentTask(*BTComp, EBTNodeResult::Succeeded);
+	}
+}
+
+AActor* UBTTaskNode::GetOwnerActor(const UGameplayTask* Task) const
+{
+	if (Task == nullptr)
+	{
+		if (IsInstanced())
+		{
+			const UBehaviorTreeComponent* BTComponent = Cast<const UBehaviorTreeComponent>(GetOuter());
+			//not having BT component for an instanced BT node is invalid!
+			check(BTComponent);
+			return BTComponent->GetAIOwner();
+		}
+		else
+		{
+			UE_LOG(LogBehaviorTree, Warning, TEXT("%s: Unable to determine Owner Actor for a null GameplayTask"), *GetName());
+			return nullptr;
+		}
+	}
+
+	const UAITask* AsAITask = Cast<const UAITask>(Task);
+	if (AsAITask)
+	{
+		return AsAITask->GetAIController();
+	}
+
+	const UGameplayTasksComponent* GTComponent = Task->GetGameplayTasksComponent();
+
+	return GTComponent ? GTComponent->GetOwnerActor(Task) : nullptr; 
+}
+
+AActor* UBTTaskNode::GetAvatarActor(const UGameplayTask* Task) const
+{
+	if (Task == nullptr)
+	{
+		if (IsInstanced())
+		{
+			const UBehaviorTreeComponent* BTComponent = Cast<const UBehaviorTreeComponent>(GetOuter());
+			//not having BT component for an instanced BT node is invalid!
+			check(BTComponent); 
+			return BTComponent->GetAIOwner() ? BTComponent->GetAIOwner()->GetPawn() : nullptr;
+		}
+		else
+		{
+			UE_LOG(LogBehaviorTree, Warning, TEXT("%s: Unable to determine Avatar Actor for a null GameplayTask"), *GetName());
+			return nullptr;
+		}
+	}
+
+	const UAITask* AsAITask = Cast<const UAITask>(Task);
+	if (AsAITask)
+	{
+		return AsAITask->GetAIController() ? AsAITask->GetAIController()->GetPawn() : nullptr;
+	}
+
+	const UGameplayTasksComponent* GTComponent = Task->GetGameplayTasksComponent();
+
+	return GTComponent ? GTComponent->GetOwnerActor(Task) : nullptr;
+}
+
+uint8 UBTTaskNode::GetDefaultPriority() const 
+{ 
+	return uint8(EAITaskPriority::AutonomousAI); 
+}
+
 //----------------------------------------------------------------------//
 // DEPRECATED
 //----------------------------------------------------------------------//
diff --git a/Engine/Source/Runtime/AIModule/Private/BehaviorTree/Tasks/BTTask_MoveDirectlyToward.cpp b/Engine/Source/Runtime/AIModule/Private/BehaviorTree/Tasks/BTTask_MoveDirectlyToward.cpp
--- a/Engine/Source/Runtime/AIModule/Private/BehaviorTree/Tasks/BTTask_MoveDirectlyToward.cpp
+++ b/Engine/Source/Runtime/AIModule/Private/BehaviorTree/Tasks/BTTask_MoveDirectlyToward.cpp
@@ -5,16 +5,18 @@
 #include "BehaviorTree/Blackboard/BlackboardKeyType_Object.h"
 #include "BehaviorTree/Blackboard/BlackboardKeyType_Vector.h"
 #include "BehaviorTree/Tasks/BTTask_MoveDirectlyToward.h"
+#include "Tasks/AITask_MoveTo.h"
 
 UBTTask_MoveDirectlyToward::UBTTask_MoveDirectlyToward(const FObjectInitializer& ObjectInitializer) 
 	: Super(ObjectInitializer)
-	, AcceptableRadius(50.f)
 	, bProjectVectorGoalToNavigation(true)
-	, bAllowStrafe(true)
-	, bStopOnOverlap(false)
 {
 	NodeName = "MoveDirectlyToward";
 
+	AcceptableRadius = GET_AI_CONFIG_VAR(AcceptanceRadius);
+	bStopOnOverlap = GET_AI_CONFIG_VAR(bFinishMoveOnGoalOverlap);
+	bAllowStrafe = GET_AI_CONFIG_VAR(bAllowStrafing);
+
 	// accept only actors and vectors
 	BlackboardKey.AddObjectFilter(this, GET_MEMBER_NAME_CHECKED(UBTTask_MoveDirectlyToward, BlackboardKey), AActor::StaticClass());
 	BlackboardKey.AddVectorFilter(this, GET_MEMBER_NAME_CHECKED(UBTTask_MoveDirectlyToward, BlackboardKey));
@@ -29,37 +31,74 @@ EBTNodeResult::Type UBTTask_MoveDirectlyToward::ExecuteTask(UBehaviorTreeCompone
 
 	if (MyController && MyBlackboard)
 	{
-		EPathFollowingRequestResult::Type RequestResult = EPathFollowingRequestResult::Failed;
-
-		if (BlackboardKey.SelectedKeyType == UBlackboardKeyType_Object::StaticClass())
+		if (GET_AI_CONFIG_VAR(bEnableBTAITasks))
 		{
-			UObject* KeyValue = MyBlackboard->GetValue<UBlackboardKeyType_Object>(BlackboardKey.GetSelectedKeyID());
-			AActor* TargetActor = Cast<AActor>(KeyValue);
-			if (TargetActor)
+			UAITask_MoveTo* AIMoveTask = NewBTAITask<UAITask_MoveTo>(OwnerComp);
+
+			if (AIMoveTask != nullptr)
 			{
-				RequestResult = bDisablePathUpdateOnGoalLocationChange ?
-					MyController->MoveToLocation(TargetActor->GetActorLocation(), AcceptableRadius, bStopOnOverlap, /*bUsePathfinding=*/false, /*bProjectDestinationToNavigation=*/bProjectVectorGoalToNavigation, bAllowStrafe) :
-					MyController->MoveToActor(TargetActor, AcceptableRadius, bStopOnOverlap, /*bUsePathfinding=*/false, bAllowStrafe);
+				bool bSetUp = false;
+				if (BlackboardKey.SelectedKeyType == UBlackboardKeyType_Object::StaticClass())
+				{
+					UObject* KeyValue = MyBlackboard->GetValue<UBlackboardKeyType_Object>(BlackboardKey.GetSelectedKeyID());
+					AActor* TargetActor = Cast<AActor>(KeyValue);
+					if (TargetActor)
+					{
+						AIMoveTask->SetUp(MyController, FVector::ZeroVector, TargetActor, AcceptableRadius, /*bUsePathfinding=*/false, FAISystem::BoolToAIOption(bStopOnOverlap));
+						NodeResult = EBTNodeResult::InProgress;
+					}
+					else
+					{
+						UE_VLOG(MyController, LogBehaviorTree, Warning, TEXT("UBTTask_MoveDirectlyToward::ExecuteTask tried to go to actor while BB %s entry was empty"), *BlackboardKey.SelectedKeyName.ToString());
+					}
+				}
+				else if (BlackboardKey.SelectedKeyType == UBlackboardKeyType_Vector::StaticClass())
+				{
+					const FVector TargetLocation = MyBlackboard->GetValue<UBlackboardKeyType_Vector>(BlackboardKey.GetSelectedKeyID());
+					AIMoveTask->SetUp(MyController, TargetLocation, nullptr, AcceptableRadius, /*bUsePathfinding=*/false, FAISystem::BoolToAIOption(bStopOnOverlap));
+					NodeResult = EBTNodeResult::InProgress;
+				}
+
+				if (NodeResult == EBTNodeResult::InProgress)
+				{
+					AIMoveTask->ReadyForActivation();
+				}
 			}
 		}
-		else if (BlackboardKey.SelectedKeyType == UBlackboardKeyType_Vector::StaticClass())
+		else
 		{
-			const FVector TargetLocation = MyBlackboard->GetValue<UBlackboardKeyType_Vector>(BlackboardKey.GetSelectedKeyID());
-			RequestResult = MyController->MoveToLocation(TargetLocation, AcceptableRadius, bStopOnOverlap, /*bUsePathfinding=*/false, /*bProjectDestinationToNavigation=*/bProjectVectorGoalToNavigation, bAllowStrafe);
-		}
+			EPathFollowingRequestResult::Type RequestResult = EPathFollowingRequestResult::Failed;
 
-		if (RequestResult == EPathFollowingRequestResult::RequestSuccessful)
-		{
-			const FAIRequestID RequestID = MyController->GetCurrentMoveRequestID();
+			if (BlackboardKey.SelectedKeyType == UBlackboardKeyType_Object::StaticClass())
+			{
+				UObject* KeyValue = MyBlackboard->GetValue<UBlackboardKeyType_Object>(BlackboardKey.GetSelectedKeyID());
+				AActor* TargetActor = Cast<AActor>(KeyValue);
+				if (TargetActor)
+				{
+					RequestResult = bDisablePathUpdateOnGoalLocationChange ?
+						MyController->MoveToLocation(TargetActor->GetActorLocation(), AcceptableRadius, bStopOnOverlap, /*bUsePathfinding=*/false, /*bProjectDestinationToNavigation=*/bProjectVectorGoalToNavigation, bAllowStrafe) :
+						MyController->MoveToActor(TargetActor, AcceptableRadius, bStopOnOverlap, /*bUsePathfinding=*/false, bAllowStrafe);
+				}
+			}
+			else if (BlackboardKey.SelectedKeyType == UBlackboardKeyType_Vector::StaticClass())
+			{
+				const FVector TargetLocation = MyBlackboard->GetValue<UBlackboardKeyType_Vector>(BlackboardKey.GetSelectedKeyID());
+				RequestResult = MyController->MoveToLocation(TargetLocation, AcceptableRadius, bStopOnOverlap, /*bUsePathfinding=*/false, /*bProjectDestinationToNavigation=*/bProjectVectorGoalToNavigation, bAllowStrafe);
+			}
 
-			MyMemory->MoveRequestID = RequestID;
-			WaitForMessage(OwnerComp, UBrainComponent::AIMessage_MoveFinished, RequestID);
+			if (RequestResult == EPathFollowingRequestResult::RequestSuccessful)
+			{
+				const FAIRequestID RequestID = MyController->GetCurrentMoveRequestID();
 
-			NodeResult = EBTNodeResult::InProgress;
-		}
-		else if (RequestResult == EPathFollowingRequestResult::AlreadyAtGoal)
-		{
-			NodeResult = EBTNodeResult::Succeeded;
+				MyMemory->MoveRequestID = RequestID;
+				WaitForMessage(OwnerComp, UBrainComponent::AIMessage_MoveFinished, RequestID);
+
+				NodeResult = EBTNodeResult::InProgress;
+			}
+			else if (RequestResult == EPathFollowingRequestResult::AlreadyAtGoal)
+			{
+				NodeResult = EBTNodeResult::Succeeded;
+			}
 		}
 	}
 
diff --git a/Engine/Source/Runtime/AIModule/Private/BehaviorTree/Tasks/BTTask_MoveTo.cpp b/Engine/Source/Runtime/AIModule/Private/BehaviorTree/Tasks/BTTask_MoveTo.cpp
--- a/Engine/Source/Runtime/AIModule/Private/BehaviorTree/Tasks/BTTask_MoveTo.cpp
+++ b/Engine/Source/Runtime/AIModule/Private/BehaviorTree/Tasks/BTTask_MoveTo.cpp
@@ -6,17 +6,19 @@
 #include "BehaviorTree/BlackboardComponent.h"
 #include "BehaviorTree/Blackboard/BlackboardKeyType_Object.h"
 #include "BehaviorTree/Blackboard/BlackboardKeyType_Vector.h"
+#include "Tasks/AITask_MoveTo.h"
 
 UBTTask_MoveTo::UBTTask_MoveTo(const FObjectInitializer& ObjectInitializer) 
 	: Super(ObjectInitializer)
-	, AcceptableRadius(50.f)
-	, bAllowStrafe(false)
-	, bAllowPartialPath(true)
-	, bStopOnOverlap(true)
 {
 	NodeName = "Move To";
 	bNotifyTick = true;
 
+	AcceptableRadius = GET_AI_CONFIG_VAR(AcceptanceRadius);
+	bStopOnOverlap = GET_AI_CONFIG_VAR(bFinishMoveOnGoalOverlap); 
+	bAllowStrafe = GET_AI_CONFIG_VAR(bAllowStrafing);
+	bAllowPartialPath = GET_AI_CONFIG_VAR(bAcceptPartialPaths);
+
 	// accept only actors and vectors
 	BlackboardKey.AddObjectFilter(this, GET_MEMBER_NAME_CHECKED(UBTTask_MoveTo, BlackboardKey), AActor::StaticClass());
 	BlackboardKey.AddVectorFilter(this, GET_MEMBER_NAME_CHECKED(UBTTask_MoveTo, BlackboardKey));
@@ -52,51 +54,88 @@ EBTNodeResult::Type UBTTask_MoveTo::PerformMoveTask(UBehaviorTreeComponent& Owne
 
 	if (MyController && MyBlackboard)
 	{
-		EPathFollowingRequestResult::Type RequestResult = EPathFollowingRequestResult::Failed;
-		
-		FAIMoveRequest MoveReq;
-		MoveReq.SetNavigationFilter(FilterClass);
-		MoveReq.SetAllowPartialPath(bAllowPartialPath);
-		MoveReq.SetAcceptanceRadius(AcceptableRadius);
-		MoveReq.SetCanStrafe(bAllowStrafe);
-		MoveReq.SetStopOnOverlap(bStopOnOverlap);
-
-		if (BlackboardKey.SelectedKeyType == UBlackboardKeyType_Object::StaticClass())
+		if (GET_AI_CONFIG_VAR(bEnableBTAITasks))
 		{
-			UObject* KeyValue = MyBlackboard->GetValue<UBlackboardKeyType_Object>(BlackboardKey.GetSelectedKeyID());
-			AActor* TargetActor = Cast<AActor>(KeyValue);
-			if (TargetActor)
-			{
-				MoveReq.SetGoalActor(TargetActor);
+			UAITask_MoveTo* AIMoveTask = NewBTAITask<UAITask_MoveTo>(OwnerComp);
 
-				RequestResult = MyController->MoveTo(MoveReq);
-			}
-			else
+			if (AIMoveTask != nullptr)
 			{
-				UE_VLOG(MyController, LogBehaviorTree, Warning, TEXT("UBTTask_MoveTo::ExecuteTask tried to go to actor while BB %s entry was empty"), *BlackboardKey.SelectedKeyName.ToString());
+				bool bSetUp = false;
+				if (BlackboardKey.SelectedKeyType == UBlackboardKeyType_Object::StaticClass())
+				{
+					UObject* KeyValue = MyBlackboard->GetValue<UBlackboardKeyType_Object>(BlackboardKey.GetSelectedKeyID());
+					AActor* TargetActor = Cast<AActor>(KeyValue);
+					if (TargetActor)
+					{
+						AIMoveTask->SetUp(MyController, FVector::ZeroVector, TargetActor, AcceptableRadius, /*bUsePathfinding=*/true, FAISystem::BoolToAIOption(bStopOnOverlap), FAISystem::BoolToAIOption(bAllowPartialPath));
+						NodeResult = EBTNodeResult::InProgress;
+					}
+					else
+					{
+						UE_VLOG(MyController, LogBehaviorTree, Warning, TEXT("UBTTask_MoveTo::ExecuteTask tried to go to actor while BB %s entry was empty"), *BlackboardKey.SelectedKeyName.ToString());
+					}
+				}
+				else if (BlackboardKey.SelectedKeyType == UBlackboardKeyType_Vector::StaticClass())
+				{
+					const FVector TargetLocation = MyBlackboard->GetValue<UBlackboardKeyType_Vector>(BlackboardKey.GetSelectedKeyID());
+					AIMoveTask->SetUp(MyController, TargetLocation, nullptr, AcceptableRadius, /*bUsePathfinding=*/true, FAISystem::BoolToAIOption(bStopOnOverlap), FAISystem::BoolToAIOption(bAllowPartialPath));
+					NodeResult = EBTNodeResult::InProgress;
+				}
+
+				if (NodeResult == EBTNodeResult::InProgress)
+				{
+					AIMoveTask->ReadyForActivation();
+				}
 			}
 		}
-		else if (BlackboardKey.SelectedKeyType == UBlackboardKeyType_Vector::StaticClass())
+		else
 		{
-			const FVector TargetLocation = MyBlackboard->GetValue<UBlackboardKeyType_Vector>(BlackboardKey.GetSelectedKeyID());
-			MoveReq.SetGoalLocation(TargetLocation);
+			EPathFollowingRequestResult::Type RequestResult = EPathFollowingRequestResult::Failed;
 
-			RequestResult = MyController->MoveTo(MoveReq);
-		}
+			FAIMoveRequest MoveReq;
+			MoveReq.SetNavigationFilter(FilterClass);
+			MoveReq.SetAllowPartialPath(bAllowPartialPath);
+			MoveReq.SetAcceptanceRadius(AcceptableRadius);
+			MoveReq.SetCanStrafe(bAllowStrafe);
+			MoveReq.SetStopOnOverlap(bStopOnOverlap);
 
-		if (RequestResult == EPathFollowingRequestResult::RequestSuccessful)
-		{
-			const FAIRequestID RequestID = MyController->GetCurrentMoveRequestID();
+			if (BlackboardKey.SelectedKeyType == UBlackboardKeyType_Object::StaticClass())
+			{
+				UObject* KeyValue = MyBlackboard->GetValue<UBlackboardKeyType_Object>(BlackboardKey.GetSelectedKeyID());
+				AActor* TargetActor = Cast<AActor>(KeyValue);
+				if (TargetActor)
+				{
+					MoveReq.SetGoalActor(TargetActor);
+
+					RequestResult = MyController->MoveTo(MoveReq);
+				}
+				else
+				{
+					UE_VLOG(MyController, LogBehaviorTree, Warning, TEXT("UBTTask_MoveTo::ExecuteTask tried to go to actor while BB %s entry was empty"), *BlackboardKey.SelectedKeyName.ToString());
+				}
+			}
+			else if (BlackboardKey.SelectedKeyType == UBlackboardKeyType_Vector::StaticClass())
+			{
+				const FVector TargetLocation = MyBlackboard->GetValue<UBlackboardKeyType_Vector>(BlackboardKey.GetSelectedKeyID());
+				MoveReq.SetGoalLocation(TargetLocation);
 
-			MyMemory->MoveRequestID = RequestID;
-			WaitForMessage(OwnerComp, UBrainComponent::AIMessage_MoveFinished, RequestID);
-			WaitForMessage(OwnerComp, UBrainComponent::AIMessage_RepathFailed);
+				RequestResult = MyController->MoveTo(MoveReq);
+			}
 
-			NodeResult = EBTNodeResult::InProgress;
-		}
-		else if (RequestResult == EPathFollowingRequestResult::AlreadyAtGoal)
-		{
-			NodeResult = EBTNodeResult::Succeeded;
+			if (RequestResult == EPathFollowingRequestResult::RequestSuccessful)
+			{
+				const FAIRequestID RequestID = MyController->GetCurrentMoveRequestID();
+
+				MyMemory->MoveRequestID = RequestID;
+				WaitForMessage(OwnerComp, UBrainComponent::AIMessage_MoveFinished, RequestID);
+				WaitForMessage(OwnerComp, UBrainComponent::AIMessage_RepathFailed);
+
+				NodeResult = EBTNodeResult::InProgress;
+			}
+			else if (RequestResult == EPathFollowingRequestResult::AlreadyAtGoal)
+			{
+				NodeResult = EBTNodeResult::Succeeded;
+			}
 		}
 	}
 
diff --git a/Engine/Source/Runtime/AIModule/Private/Perception/AIPerceptionComponent.cpp b/Engine/Source/Runtime/AIModule/Private/Perception/AIPerceptionComponent.cpp
--- a/Engine/Source/Runtime/AIModule/Private/Perception/AIPerceptionComponent.cpp
+++ b/Engine/Source/Runtime/AIModule/Private/Perception/AIPerceptionComponent.cpp
@@ -128,6 +128,9 @@ void UAIPerceptionComponent::SetMaxStimulusAge(int32 ConfigIndex, float MaxAge)
 		MaxActiveAge.AddUninitialized(ConfigIndex - MaxActiveAge.Num() + 1);
 	}
 	MaxActiveAge[ConfigIndex] = MaxAge;
+
+	// @todo process all data already gathered and see if any _still_active_ stimuli
+	// got it's expiration prolonged, with SetExpirationAge
 }
 
 void UAIPerceptionComponent::OnRegister()
@@ -369,7 +372,8 @@ AActor* UAIPerceptionComponent::GetMutableBodyActor()
 
 void UAIPerceptionComponent::RegisterStimulus(AActor* Source, const FAIStimulus& Stimulus)
 {
-	StimuliToProcess.Add(FStimulusToProcess(Source, Stimulus));
+	FStimulusToProcess& StimulusToProcess = StimuliToProcess[StimuliToProcess.Add(FStimulusToProcess(Source, Stimulus))];
+	StimulusToProcess.Stimulus.SetExpirationAge(MaxActiveAge[int32(Stimulus.Type)]);
 }
 
 void UAIPerceptionComponent::ProcessStimuli()
@@ -427,15 +431,14 @@ void UAIPerceptionComponent::ProcessStimuli()
 		{
 			RefreshStimulus(StimulusStore, SourcedStimulus->Stimulus);
 		}
-		else if (StimulusStore.IsExpired())
-		{
-			HandleExpiredStimulus(StimulusStore);
-		}
-		else
+		else if (StimulusStore.IsExpired() == false)
 		{
+			//HandleExpiredStimulus(StimulusStore);
+
 			// @note there some more valid info in SourcedStimulus->Stimulus regarding test that failed
 			// may be useful in future
 			StimulusStore.MarkNoLongerSensed();
+			StimulusStore.SetStimulusAge(0);
 		}
 
 		// if the new stimulus is "valid" or it's info that "no longer sensed" and it used to be sensed successfully
@@ -470,12 +473,13 @@ void UAIPerceptionComponent::RefreshStimulus(FAIStimulus& StimulusStore, const F
 	if (NewStimulus.GetAge() <= StimulusStore.GetAge() || StimulusStore.Strength < NewStimulus.Strength)
 	{
 		StimulusStore = NewStimulus;
+		// update stimulus 
 	}
 }
 
 void UAIPerceptionComponent::HandleExpiredStimulus(FAIStimulus& StimulusStore)
 {
-	ensure(StimulusStore.IsExpired() == true && StimulusStore.WasSuccessfullySensed() == false && StimulusStore.IsActive() == false);
+	ensure(StimulusStore.IsExpired() == true);
 }
 
 bool UAIPerceptionComponent::AgeStimuli(const float ConstPerceptionAgingRate)
@@ -489,9 +493,9 @@ bool UAIPerceptionComponent::AgeStimuli(const float ConstPerceptionAgingRate)
 		for (FAIStimulus& Stimulus : ActorPerceptionInfo.LastSensedStimuli)
 		{
 			// Age the stimulus. If it is active but has just expired, mark it as such
-			if (Stimulus.AgeStimulus(ConstPerceptionAgingRate) == false && 
-				Stimulus.IsActive() && 
-				!Stimulus.IsExpired())
+			if (Stimulus.AgeStimulus(ConstPerceptionAgingRate) == false 
+				&& (Stimulus.IsActive() || Stimulus.WantsToNotifyOnlyOnPerceptionChange())
+				&& Stimulus.IsExpired() == false)
 			{
 				AActor* TargetActor = ActorPerceptionInfo.Target.Get();
 				if (TargetActor)
diff --git a/Engine/Source/Runtime/AIModule/Private/Tasks/AITask.cpp b/Engine/Source/Runtime/AIModule/Private/Tasks/AITask.cpp
new file mode 100644
--- /dev/null
+++ b/Engine/Source/Runtime/AIModule/Private/Tasks/AITask.cpp
@@ -0,0 +1,61 @@
+// Copyright 1998-2015 Epic Games, Inc. All Rights Reserved.
+
+#include "AIModulePrivate.h"
+#include "GameplayTasksComponent.h"
+
+
+UAITask::UAITask(const FObjectInitializer& ObjectInitializer)
+	: Super(ObjectInitializer)
+{
+	Priority = uint8(EAITaskPriority::AutonomousAI);
+}
+
+AAIController* UAITask::GetAIControllerForActor(AActor* Actor)
+{
+	AAIController* Result = Cast<AAIController>(Actor);
+
+	if (Result == nullptr)
+	{
+		APawn* AsPawn = Cast<APawn>(Actor);
+		if (AsPawn != nullptr)
+		{
+			Result = Cast<AAIController>(AsPawn->GetController());
+		}
+	}
+
+	return Result;
+}
+
+void UAITask::Activate()
+{
+	Super::Activate();
+
+	if (OwnerController == nullptr)
+	{
+		AActor* OwnerActor = GetOwnerActor();
+		if (OwnerActor)
+		{
+			OwnerController = GetAIControllerForActor(OwnerActor);
+		}
+		else
+		{
+			// ERROR!
+		}
+	}
+}
+
+void UAITask::InitAITask(AAIController& AIOwner, IGameplayTaskOwnerInterface& InTaskOwner, uint8 InPriority)
+{
+	OwnerController = &AIOwner;
+	InitTask(InTaskOwner, InPriority);
+
+	if (InPriority == uint8(EAITaskPriority::AutonomousAI))
+	{
+		AddRequiredResource<UAIResource_Logic>();
+	}
+}
+
+void UAITask::RequestAILogicLocking()
+{
+	AddClaimedResource<UAIResource_Logic>();
+}
\ No newline at end of file
diff --git a/Engine/Source/Runtime/AIModule/Private/Tasks/AITask_MoveTo.cpp b/Engine/Source/Runtime/AIModule/Private/Tasks/AITask_MoveTo.cpp
--- a/Engine/Source/Runtime/AIModule/Private/Tasks/AITask_MoveTo.cpp
+++ b/Engine/Source/Runtime/AIModule/Private/Tasks/AITask_MoveTo.cpp
@@ -8,14 +8,18 @@ UAITask_MoveTo::UAITask_MoveTo(const FObjectInitializer& ObjectInitializer)
 	: Super(ObjectInitializer)
 {
 	bIsPausable = true;
-}
 
-UAITask_MoveTo* UAITask_MoveTo::PawnMoveTo(APawn* Pawn, FVector Destination, AActor* TargetActor, float AcceptanceRadius, bool bStopOnOverlap, bool bAcceptPartialPath)
-{
-	return nullptr;
+	RealGoalLocation = FAISystem::InvalidLocation;
+	MoveGoalActor = nullptr;
+	MoveAcceptanceRadius = GET_AI_CONFIG_VAR(AcceptanceRadius);
+	bShouldStopOnOverlap = GET_AI_CONFIG_VAR(bFinishMoveOnGoalOverlap);
+	bShouldAcceptPartialPath = GET_AI_CONFIG_VAR(bAcceptPartialPaths);
+	bShouldUsePathfinding = true;
+
+	AddRequiredResource(UAIResource_Movement::StaticClass());
 }
 
-UAITask_MoveTo* UAITask_MoveTo::AIMoveTo(AAIController* Controller, FVector TargetLocation, AActor* TargetActor, float AcceptanceRadius, bool bStopOnOverlap, bool bAcceptPartialPath)
+UAITask_MoveTo* UAITask_MoveTo::AIMoveTo(AAIController* Controller, FVector InGoalLocation, AActor* InGoalActor, float AcceptanceRadius, EAIOptionFlag::Type StopOnOverlap, EAIOptionFlag::Type AcceptPartialPath, bool bUsePathfinding, bool bLockAILogic)
 {
 	if (Controller == nullptr)
 	{
@@ -25,28 +29,52 @@ UAITask_MoveTo* UAITask_MoveTo::AIMoveTo(AAIController* Controller, FVector Targ
 	UAITask_MoveTo* MyTask = NewTask<UAITask_MoveTo>(*static_cast<IGameplayTaskOwnerInterface*>(Controller));
 	if (MyTask)
 	{
-		MyTask->OwnerController = Controller;
-
-		if (TargetActor != nullptr)
-		{
-			MyTask->MoveRequest = FAIMoveRequest(TargetActor);
-		}
-		else
+		MyTask->SetUp(Controller, InGoalLocation, InGoalActor, AcceptanceRadius, bUsePathfinding, StopOnOverlap, AcceptPartialPath);
+		MyTask->Priority = uint8(EAITaskPriority::High);
+		if (bLockAILogic)
 		{
-			MyTask->MoveRequest = FAIMoveRequest(TargetLocation);
+			MyTask->RequestAILogicLocking();
 		}
+	}
 
-		MyTask->MoveRequest.SetAllowPartialPath(bAcceptPartialPath)
-			.SetAcceptanceRadius(AcceptanceRadius)
-			.SetStopOnOverlap(bStopOnOverlap);
-		//MyTask->MoveReq.SetUsePathfinding(bUsePathfinding);
-		//MyTask->MoveRequest.SetNavigationFilter(FilterClass);
-		//MyTask->MoveRequest.SetCanStrafe(bCanStrafe);
+	return MyTask;
+}
 
-		MyTask->RequireResource(UAIResource_Movement::StaticClass());
+void UAITask_MoveTo::SetUp(AAIController* Controller, FVector InGoalLocation, AActor* InGoalActor, float AcceptanceRadius, bool bUsePathfinding, EAIOptionFlag::Type StopOnOverlap, EAIOptionFlag::Type AcceptPartialPath)
+{
+	OwnerController = Controller;
+
+	if (InGoalActor != nullptr)
+	{
+		MoveGoalActor = InGoalActor;
+	}
+	else
+	{
+		RealGoalLocation = InGoalLocation;
 	}
 
-	return MyTask;
+	if (AcceptanceRadius >= 0)
+	{
+		MoveAcceptanceRadius = AcceptanceRadius;
+	}
+	
+	bShouldStopOnOverlap = FAISystem::PickAIOption(StopOnOverlap, bShouldStopOnOverlap);
+	bShouldUsePathfinding = bUsePathfinding;
+
+	if (AcceptPartialPath != EAIOptionFlag::Default)
+	{
+		bShouldAcceptPartialPath = (AcceptPartialPath == EAIOptionFlag::Enable);
+	}
+}
+
+void UAITask_MoveTo::PostInitProperties()
+{
+	Super::PostInitProperties();
+
+	if (MoveGoalActor == nullptr)
+	{
+		RealGoalLocation = MoveGoalLocation;
+	}
 }
 
 void UAITask_MoveTo::HandleMoveFinished(FAIRequestID RequestID, EPathFollowingResult::Type Result)
@@ -72,6 +100,15 @@ void UAITask_MoveTo::HandleMoveFinished(FAIRequestID RequestID, EPathFollowingRe
 
 void UAITask_MoveTo::Activate()
 {
+	Super::Activate();
+
+	FAIMoveRequest MoveRequest = (MoveGoalActor != nullptr) ? FAIMoveRequest(MoveGoalActor) : FAIMoveRequest(RealGoalLocation);
+
+	MoveRequest.SetAllowPartialPath(bShouldAcceptPartialPath)
+		.SetAcceptanceRadius(MoveAcceptanceRadius)
+		.SetStopOnOverlap(bShouldStopOnOverlap)
+		.SetUsePathfinding(bShouldUsePathfinding);
+
 	const EPathFollowingRequestResult::Type RequestResult = OwnerController->MoveTo(MoveRequest);
 
 	switch (RequestResult)
@@ -106,6 +143,7 @@ void UAITask_MoveTo::Pause()
 	if (OwnerController)
 	{
 		OwnerController->GetPathFollowingComponent()->OnMoveFinished.RemoveAll(this);
+		OwnerController->PauseMove(MoveRequestID);
 	}
 	Super::Pause();
 }
@@ -115,3 +153,13 @@ void UAITask_MoveTo::Resume()
 	Activate();
 	Super::Resume();
 }
+
+void UAITask_MoveTo::OnDestroy(bool bOwnerFinished)
+{
+	Super::OnDestroy(bOwnerFinished);
+	if (OwnerController && OwnerController->GetPathFollowingComponent())
+	{
+		OwnerController->GetPathFollowingComponent()->OnMoveFinished.RemoveAll(this);
+		OwnerController->GetPathFollowingComponent()->AbortMove(TEXT("UAITask_MoveTo finishing"), MoveRequestID);
+	}
+}
diff --git a/Engine/Source/Runtime/GameplayAbilities/Private/Abilities/GameplayAbility.cpp b/Engine/Source/Runtime/GameplayAbilities/Private/Abilities/GameplayAbility.cpp
--- a/Engine/Source/Runtime/GameplayAbilities/Private/Abilities/GameplayAbility.cpp
+++ b/Engine/Source/Runtime/GameplayAbilities/Private/Abilities/GameplayAbility.cpp
@@ -1006,18 +1006,18 @@ FGameplayAbilityTargetingLocationInfo UGameplayAbility::MakeTargetLocationInfoFr
 
 //----------------------------------------------------------------------
 
-UGameplayTasksComponent* UGameplayAbility::GetGameplayTasksComponent()
+UGameplayTasksComponent* UGameplayAbility::GetGameplayTasksComponent(const UGameplayTask& Task) const
 {
 	return GetCurrentActorInfo() ? GetCurrentActorInfo()->AbilitySystemComponent.Get() : nullptr;
 }
 
-AActor* UGameplayAbility::GetOwnerActor() const
+AActor* UGameplayAbility::GetOwnerActor(const UGameplayTask* Task) const
 {
 	const FGameplayAbilityActorInfo* Info = GetCurrentActorInfo();
 	return Info ? Info->OwnerActor.Get() : nullptr;
 }
 
-AActor* UGameplayAbility::GetAvatarActor() const
+AActor* UGameplayAbility::GetAvatarActor(const UGameplayTask* Task) const
 {
 	const FGameplayAbilityActorInfo* Info = GetCurrentActorInfo();
 	return Info ? Info->AvatarActor.Get() : nullptr;
diff --git a/Engine/Source/Runtime/GameplayAbilities/Private/AbilitySystemComponent.cpp b/Engine/Source/Runtime/GameplayAbilities/Private/AbilitySystemComponent.cpp
--- a/Engine/Source/Runtime/GameplayAbilities/Private/AbilitySystemComponent.cpp
+++ b/Engine/Source/Runtime/GameplayAbilities/Private/AbilitySystemComponent.cpp
@@ -1275,7 +1275,7 @@ TArray<TWeakObjectPtr<UGameplayTask> >&	UAbilitySystemComponent::GetAbilityActiv
 	return Ability->ActiveTasks;
 }
 
-AActor* UAbilitySystemComponent::GetAvatarActor() const
+AActor* UAbilitySystemComponent::GetAvatarActor(const UGameplayTask* Task) const
 {
 	check(AbilityActorInfo.IsValid());
 	return AbilityActorInfo->OwnerActor.Get();
diff --git a/Engine/Source/Runtime/GameplayAbilities/Public/Abilities/GameplayAbility.h b/Engine/Source/Runtime/GameplayAbilities/Public/Abilities/GameplayAbility.h
--- a/Engine/Source/Runtime/GameplayAbilities/Public/Abilities/GameplayAbility.h
+++ b/Engine/Source/Runtime/GameplayAbilities/Public/Abilities/GameplayAbility.h
@@ -260,9 +260,9 @@ public:
 	/** Called by an ability task, originating from this ability, when it ends */
 	virtual void OnTaskDeactivated(UGameplayTask& Task) override;
 
-	virtual UGameplayTasksComponent* GetGameplayTasksComponent() override;
-	virtual AActor* GetOwnerActor() const override;
-	virtual AActor* GetAvatarActor() const override;
+	virtual UGameplayTasksComponent* GetGameplayTasksComponent(const UGameplayTask& Task) const override;
+	virtual AActor* GetOwnerActor(const UGameplayTask* Task) const override;
+	virtual AActor* GetAvatarActor(const UGameplayTask* Task) const override;
 
 	// --------------------------------------
 	//	Input
diff --git a/Engine/Source/Runtime/GameplayAbilities/Public/Abilities/Tasks/AbilityTask.h b/Engine/Source/Runtime/GameplayAbilities/Public/Abilities/Tasks/AbilityTask.h
--- a/Engine/Source/Runtime/GameplayAbilities/Public/Abilities/Tasks/AbilityTask.h
+++ b/Engine/Source/Runtime/GameplayAbilities/Public/Abilities/Tasks/AbilityTask.h
@@ -94,7 +94,7 @@ class GAMEPLAYABILITIES_API UAbilityTask : public UGameplayTask
 
 		T* MyObj = NewObject<T>();
 		UGameplayAbility* ThisAbility = CastChecked<UGameplayAbility>(WorldContextObject);
-		MyObj->InitTask(*ThisAbility);
+		MyObj->InitTask(*ThisAbility, ThisAbility->GetDefaultPriority());
 		MyObj->InstanceName = InstanceName;
 		return MyObj;
 	}
diff --git a/Engine/Source/Runtime/GameplayAbilities/Public/AbilitySystemComponent.h b/Engine/Source/Runtime/GameplayAbilities/Public/AbilitySystemComponent.h
--- a/Engine/Source/Runtime/GameplayAbilities/Public/AbilitySystemComponent.h
+++ b/Engine/Source/Runtime/GameplayAbilities/Public/AbilitySystemComponent.h
@@ -1037,7 +1037,7 @@ public:
 	// -
 	// UGameplayTasksComponent
 	// -
-	virtual AActor* GetAvatarActor() const override;
+	virtual AActor* GetAvatarActor(const UGameplayTask* Task) const override;
 
 protected:
 
diff --git a/Engine/Source/Runtime/GameplayTasks/Classes/GameplayTask.h b/Engine/Source/Runtime/GameplayTasks/Classes/GameplayTask.h
--- a/Engine/Source/Runtime/GameplayTasks/Classes/GameplayTask.h
+++ b/Engine/Source/Runtime/GameplayTasks/Classes/GameplayTask.h
@@ -30,6 +30,105 @@ enum class ETaskResourceOverlapPolicy : uint8
 	StartAtEnd,
 };
 	
+USTRUCT(BlueprintType)
+struct GAMEPLAYTASKS_API FGameplayResourceSet
+{
+	GENERATED_USTRUCT_BODY()
+
+	typedef uint16 FFlagContainer;
+	typedef uint8 FResourceID;
+
+	enum
+	{
+		MaxResources = sizeof(FFlagContainer)* 8
+	};
+
+private:
+	FFlagContainer Flags;
+
+public:
+	/** Mind that this constructor takes _flags_ not individual IDs */
+	explicit FGameplayResourceSet(FFlagContainer InFlags = 0) : Flags(InFlags)
+	{}
+
+	FFlagContainer GetFlags() const
+	{
+		return Flags;
+	}
+	bool IsEmpty() const
+	{
+		return Flags == 0;
+	}
+	FGameplayResourceSet& AddID(uint8 ResourceID)
+	{
+		ensure(ResourceID < MaxResources);
+		Flags |= (1 << ResourceID);
+		return *this;
+	}
+	FGameplayResourceSet& RemoveID(uint8 ResourceID)
+	{
+		ensure(ResourceID < MaxResources);
+		Flags &= ~(1 << ResourceID);
+		return *this;
+	}
+	bool HasID(uint8 ResourceID) const
+	{
+		ensure(ResourceID < MaxResources);
+		return (Flags & (1 << ResourceID)) != 0;
+	}
+	FGameplayResourceSet& AddSet(FGameplayResourceSet Other)
+	{
+		Flags |= Other.Flags;
+		return *this;
+	}
+	FGameplayResourceSet& RemoveSet(FGameplayResourceSet Other)
+	{
+		Flags &= ~Other.Flags;
+		return *this;
+	}
+	void Clear()
+	{
+		Flags = FFlagContainer(0);
+	}
+	bool HasAllIDs(FGameplayResourceSet Other) const
+	{
+		return (Flags & Other.Flags) == Flags;
+	}
+	bool HasAnyID(FGameplayResourceSet Other) const
+	{
+		return (Flags & Other.Flags) != 0;
+	}
+	FGameplayResourceSet GetOverlap(FGameplayResourceSet Other) const
+	{
+		return FGameplayResourceSet(Flags & Other.Flags);
+	}
+	FGameplayResourceSet GetDifference(FGameplayResourceSet Other) const
+	{
+		return FGameplayResourceSet(Flags & ~(Flags & Other.Flags));
+	}
+
+	bool operator==(const FGameplayResourceSet& Other) const
+	{
+		return Flags == Other.Flags;
+	}
+
+	bool operator!=(const FGameplayResourceSet& Other) const
+	{
+		return Flags != Other.Flags;
+	}
+
+	static FGameplayResourceSet AllResources()
+	{
+		return FGameplayResourceSet(FFlagContainer(-1));
+	}
+
+	static FGameplayResourceSet NoResources()
+	{
+		return FGameplayResourceSet(FFlagContainer(0));
+	}
+
+	FString GetDebugDescription() const;
+};
 
 UCLASS(Abstract)
 class GAMEPLAYTASKS_API UGameplayTask : public UObject
@@ -52,7 +151,7 @@ protected:
 	virtual void Activate();
 
 	/** Initailizes the task with the task owner interface instance but does not actviate until Activate() is called */
-	virtual void InitTask(IGameplayTaskOwnerInterface& InTaskOwner);
+	void InitTask(IGameplayTaskOwnerInterface& InTaskOwner, uint8 InPriority);
 
 public:
 	virtual void InitSimulatedTask(UGameplayTasksComponent& InGameplayTasksComponent);
@@ -117,26 +216,43 @@ public:
 	FORCEINLINE uint8 GetPriority() const { return Priority; }
 	FORCEINLINE bool RequiresPriorityOrResourceManagement() const { return bCaresAboutPriority == true || RequiredResources.IsEmpty() == false; }
 	FORCEINLINE FGameplayResourceSet GetRequiredResources() { return RequiredResources; }
+	FORCEINLINE FGameplayResourceSet GetClaimedResources() { return ClaimedResources; }
 	
 	FORCEINLINE EGameplayTaskState GetState() const { return TaskState; }
 	FORCEINLINE bool IsActive() const { return (TaskState == EGameplayTaskState::Active); }
 	
 	IGameplayTaskOwnerInterface* GetTaskOwner() const { return TaskOwner.IsValid() ? &(*TaskOwner) : nullptr; }
 	UGameplayTasksComponent* GetGameplayTasksComponent() { return TasksComponent.Get(); }
+	UGameplayTasksComponent* GetGameplayTasksComponent() const { return TasksComponent.Get(); }
 	bool IsOwnedByTasksComponent() const { return bOwnedByTasksComponent; }
 
 	template <class T>
-	inline void RequireResource()
+	inline void AddRequiredResource()
+	{
+		AddRequiredResource(T::StaticClass());
+	}
+
+	template <class T>
+	inline void AddClaimedResource()
 	{
-		RequireResource(T::StaticClass());
+		AddClaimedResource(T::StaticClass());
 	}
 
 	/** Marks this task as requiring specified resource which has a number of consequences,
-	*	like task not being able to run if the resource is already taken.
-	*
-	*	@node: Calling this function makes sense only until the task is being passed over to the GameplayTasksComponent.
-	*	Once that's that resources data is consumed and further changes won't get applied */
-	void RequireResource(TSubclassOf<UGameplayTaskResource> RequiredResource);
+	 *	like task not being able to run if the resource is already taken.
+	 *
+	 *	@note: Calling this function makes sense only until the task is being passed over to the GameplayTasksComponent.
+	 *	Once that's that resources data is consumed and further changes won't get applied 
+	 */
+	void AddRequiredResource(TSubclassOf<UGameplayTaskResource> RequiredResource);
+	void AddRequiredResourceSet(const TArray<TSubclassOf<UGameplayTaskResource> >& RequiredResourceSet);
+	void AddRequiredResourceSet(FGameplayResourceSet RequiredResourceSet);
+
+	/** 
+	 */
+	void AddClaimedResource(TSubclassOf<UGameplayTaskResource> ClaimedResource);
+	void AddClaimedResourceSet(const TArray<TSubclassOf<UGameplayTaskResource> >& AdditionalResourcesToClaim);
+	void AddClaimedResourceSet(FGameplayResourceSet AdditionalResourcesToClaim);
 
 	ETaskResourceOverlapPolicy GetResourceOverlapPolicy() const { return ResourceOverlapPolicy; }
 
@@ -192,13 +308,17 @@ protected:
 
 	/** this is set to avoid duplicate calls to task's owner and TasksComponent when both are the same object */
 	uint32 bOwnedByTasksComponent : 1;
+
+	uint32 bClaimRequiredResources : 1;
 	
-	/** Abstract "resource" IDs this task claims it needs. Note that we do not support dynamic changes to this
-	 *	variable. Once a task gets submitted to GameplayTasksComponent it's settled 
-	 *	@NOTE: if dynamic changes were needed though look for GetRequiredResources occurrences in GameplayTasksComponent
-	 *	and replace all uses of cached values with actual GetRequiredResources calls */
+	/** Abstract "resource" IDs this task needs available to be able to get activated. */
 	FGameplayResourceSet RequiredResources;
 
+	/**
+	 *	Resources that are going to be locked when this task gets activated, but are not required to get this task started
+	 */
+	FGameplayResourceSet ClaimedResources;
+
 	/** Task Owner that created us */
 	TWeakInterfacePtr<IGameplayTaskOwnerInterface> TaskOwner;
 
@@ -231,9 +351,8 @@ template <class T>
 T* UGameplayTask::NewTask(IGameplayTaskOwnerInterface& TaskOwner, FName InstanceName)
 {
 	T* MyObj = NewObject<T>();
-	MyObj->InitTask(TaskOwner);
 	MyObj->InstanceName = InstanceName;
-	MyObj->Priority = TaskOwner.GetDefaultPriority();
+	MyObj->InitTask(TaskOwner, TaskOwner.GetDefaultPriority());
 	return MyObj;
 }
 
diff --git a/Engine/Source/Runtime/GameplayTasks/Classes/GameplayTaskOwnerInterface.h b/Engine/Source/Runtime/GameplayTasks/Classes/GameplayTaskOwnerInterface.h
--- a/Engine/Source/Runtime/GameplayTasks/Classes/GameplayTaskOwnerInterface.h
+++ b/Engine/Source/Runtime/GameplayTasks/Classes/GameplayTaskOwnerInterface.h
@@ -23,12 +23,12 @@ class GAMEPLAYTASKS_API IGameplayTaskOwnerInterface
 	GENERATED_BODY()
 public:
 	virtual void OnTaskInitialized(UGameplayTask& Task);
-	virtual UGameplayTasksComponent* GetGameplayTasksComponent() = 0;
+	virtual UGameplayTasksComponent* GetGameplayTasksComponent(const UGameplayTask& Task) const = 0;
 	/** this gets called both when task starts and when task gets resumed. Check Task.GetStatus() if you want to differenciate */
 	virtual void OnTaskActivated(UGameplayTask& Task) = 0;
 	/** this gets called both when task finished and when task gets paused. Check Task.GetStatus() if you want to differenciate */
 	virtual void OnTaskDeactivated(UGameplayTask& Task) = 0;
-	virtual AActor* GetOwnerActor() const = 0;
-	virtual AActor* GetAvatarActor() const;
+	virtual AActor* GetOwnerActor(const UGameplayTask* Task) const = 0;
+	virtual AActor* GetAvatarActor(const UGameplayTask* Task) const;
 	virtual uint8 GetDefaultPriority() const { return FGameplayTasks::DefaultPriority; }
 };
\ No newline at end of file
diff --git a/Engine/Source/Runtime/GameplayTasks/Classes/GameplayTaskResource.h b/Engine/Source/Runtime/GameplayTasks/Classes/GameplayTaskResource.h
--- a/Engine/Source/Runtime/GameplayTasks/Classes/GameplayTaskResource.h
+++ b/Engine/Source/Runtime/GameplayTasks/Classes/GameplayTaskResource.h
@@ -3,7 +3,7 @@
 
 #include "GameplayTaskResource.generated.h"
 
-UCLASS(Abstract, config="Game")
+UCLASS(Abstract, config = "Game", hidedropdown)
 class GAMEPLAYTASKS_API UGameplayTaskResource : public UObject
 {
 	GENERATED_BODY()
@@ -29,6 +29,12 @@ public:
 		return bManuallySetID || (ManualResourceID != INDEX_NONE) ? ManualResourceID : AutoResourceID;
 	}
 
+	template <class T>
+	static uint8 GetResourceID()
+	{
+		return GetDefault<T>()->GetResourceID();
+	}
+
 	static uint8 GetResourceID(TSubclassOf<UGameplayTaskResource>& RequiredResource)
 	{
 		return RequiredResource->GetDefaultObject<UGameplayTaskResource>()->GetResourceID();
diff --git a/Engine/Source/Runtime/GameplayTasks/Classes/GameplayTasksComponent.h b/Engine/Source/Runtime/GameplayTasks/Classes/GameplayTasksComponent.h
--- a/Engine/Source/Runtime/GameplayTasks/Classes/GameplayTasksComponent.h
+++ b/Engine/Source/Runtime/GameplayTasks/Classes/GameplayTasksComponent.h
@@ -10,12 +10,26 @@
 class FOutBunch;
 class UActorChannel;
 
-enum class EGameplayTaskEvent
+enum class EGameplayTaskEvent : uint8
 {
 	Add,
 	Remove,
 };
 
+UENUM()
+enum class EGameplayTaskRunResult : uint8
+{
+	/** When tried running a null-task*/
+	Error,
+	Failed,
+	/** Successfully registered for running, but currently paused due to higher priority tasks running */
+	Success_Paused,
+	/** Successfully activated */
+	Success_Active,
+	/** Successfully activated, but finished instantly */
+	Success_Finished,
+};
+
 struct FGameplayTaskEventData
 {
 	EGameplayTaskEvent Event;
@@ -28,6 +42,8 @@ struct FGameplayTaskEventData
 	}
 };
 
+DECLARE_DYNAMIC_MULTICAST_DELEGATE_TwoParams(FOnClaimedResourcesChangeSignature, FGameplayResourceSet, NewlyClaimed, FGameplayResourceSet, FreshlyReleased);
+
 /**
 *	The core ActorComponent for interfacing with the GameplayAbilities System
 */
@@ -43,7 +59,7 @@ protected:
 
 	UPROPERTY()
 	TArray<UGameplayTask*> TaskPriorityQueue;
-
+	
 	/** Transient array of events whose main role is to avoid
 	 *	long chain of recurrent calls if an activated/paused/removed task 
 	 *	wants to push/pause/kill other tasks.
@@ -57,9 +73,12 @@ protected:
 	uint8 TopActivePriority;
 
 	/** Resources used by currently active tasks */
-	FGameplayResourceSet CurrentlyUsedResources;
+	FGameplayResourceSet CurrentlyClaimedResources;
 
 public:
+	UPROPERTY(BlueprintReadWrite, Category = "Gameplay Tasks")
+	FOnClaimedResourcesChangeSignature OnClaimedResourcesChange;
+
 	UGameplayTasksComponent(const FObjectInitializer& ObjectInitializer);
 	
 	UFUNCTION()
@@ -75,22 +94,28 @@ public:
 	*	activity into consideration*/
 	virtual bool GetShouldTick() const;
 	
-	virtual UGameplayTasksComponent* GetGameplayTasksComponent() override { return this; }
+	virtual UGameplayTasksComponent* GetGameplayTasksComponent(const UGameplayTask& Task) const override { return const_cast<UGameplayTasksComponent*>(this); }
 	/** Adds a task (most often an UAbilityTask) to the list of tasks to be ticked */
 	virtual void OnTaskActivated(UGameplayTask& Task) override;
 	/** Removes a task (most often an UAbilityTask) task from the list of tasks to be ticked */
 	virtual void OnTaskDeactivated(UGameplayTask& Task) override;
 
-	virtual AActor* GetOwnerActor() const override { return GetOwner(); }
-	virtual AActor* GetAvatarActor() const override;
+	virtual AActor* GetOwnerActor(const UGameplayTask* Task) const override { return GetOwner(); }
+	virtual AActor* GetAvatarActor(const UGameplayTask* Task) const override;
 		
 	/** processes the task and figures out if it should get triggered instantly or wait
 	 *	based on task's RequiredResources, Priority and ResourceOverlapPolicy */
 	void AddTaskReadyForActivation(UGameplayTask& NewTask);
 
 	void RemoveResourceConsumingTask(UGameplayTask& Task);
+	void EndAllResourceConsumingTasksOwnedBy(const IGameplayTaskOwnerInterface& TaskOwner);
+
+	FORCEINLINE FGameplayResourceSet GetCurrentlyUsedResources() const { return CurrentlyClaimedResources; }
 
-	FORCEINLINE FGameplayResourceSet GetCurrentlyUsedResources() const { return CurrentlyUsedResources; }
+	UFUNCTION(BlueprintCallable, DisplayName="Run Gameplay Task", Category = "Gameplay Tasks", meta = (AutoCreateRefTerm = "AdditionalResourcesToBeLockedWhenRun, AdditionalClaimedResources", AdvancedDisplay = "AdditionalResourcesToBeLockedWhenRun, AdditionalClaimedResources"))
+	static EGameplayTaskRunResult K2_RunGameplayTask(TScriptInterface<IGameplayTaskOwnerInterface> TaskOwner, UGameplayTask* Task, uint8 Priority, TArray<TSubclassOf<UGameplayTaskResource> > AdditionalRequiredResources, TArray<TSubclassOf<UGameplayTaskResource> > AdditionalClaimedResources);
+
+	static EGameplayTaskRunResult RunGameplayTask(IGameplayTaskOwnerInterface& TaskOwner, UGameplayTask& Task, uint8 Priority, FGameplayResourceSet AdditionalRequiredResources, FGameplayResourceSet AdditionalClaimedResources);
 
 #if ENABLE_VISUAL_LOG
 	static FString GetTaskStateName(EGameplayTaskState Value);
@@ -100,7 +125,9 @@ public:
 protected:
 	void RequestTicking();
 	void ProcessTaskEvents();
-	void UpdateTaskActivationFromIndex(int32 StartingIndex, FGameplayResourceSet ResourcesUsedByActiveUpToIndex, FGameplayResourceSet ResourcesRequiredUpToIndex);
+	void UpdateTaskActivationFromIndex(int32 StartingIndex, FGameplayResourceSet ResourcesClaimedUpToIndex, FGameplayResourceSet ResourcesBlockedUpToIndex);
+
+	void SetCurrentlyClaimedResources(FGameplayResourceSet NewClaimedSet);
 
 private:
 	/** called when a task gets ended with an external call, meaning not coming from UGameplayTasksComponent mechanics */
diff --git a/Engine/Source/Runtime/GameplayTasks/Private/GameplayTask.cpp b/Engine/Source/Runtime/GameplayTasks/Private/GameplayTask.cpp
--- a/Engine/Source/Runtime/GameplayTasks/Private/GameplayTask.cpp
+++ b/Engine/Source/Runtime/GameplayTasks/Private/GameplayTask.cpp
@@ -12,8 +12,11 @@ UGameplayTask::UGameplayTask(const FObjectInitializer& ObjectInitializer)
 	bSimulatedTask = false;
 	bIsSimulating = false;
 	bOwnedByTasksComponent = false;
+	bClaimRequiredResources = true;
 	TaskState = EGameplayTaskState::Uninitialized;
 	ResourceOverlapPolicy = ETaskResourceOverlapPolicy::StartOnTop;
+	Priority = FGameplayTasks::DefaultPriority;
+
 	SetFlags(RF_StrongRefOnFrame);
 }
 
@@ -62,16 +65,22 @@ void UGameplayTask::ReadyForActivation()
 	}
 }
 
-void UGameplayTask::InitTask(IGameplayTaskOwnerInterface& InTaskOwner)
+void UGameplayTask::InitTask(IGameplayTaskOwnerInterface& InTaskOwner, uint8 InPriority)
 {
+	Priority = InPriority;
 	TaskOwner = InTaskOwner;
-	UGameplayTasksComponent* GTComponent = InTaskOwner.GetGameplayTasksComponent();
+	UGameplayTasksComponent* GTComponent = InTaskOwner.GetGameplayTasksComponent(*this);
 	TasksComponent = GTComponent;
 
 	bOwnedByTasksComponent = (TaskOwner == GTComponent);
-
+	
 	TaskState = EGameplayTaskState::AwaitingActivation;
 
+	if (bClaimRequiredResources)
+	{
+		ClaimedResources.AddSet(RequiredResources);
+	}
+
 	InTaskOwner.OnTaskInitialized(*this);
 	if (bOwnedByTasksComponent == false && GTComponent != nullptr)
 	{
@@ -99,11 +108,11 @@ AActor* UGameplayTask::GetOwnerActor() const
 {
 	if (TaskOwner.IsValid())
 	{
-		return TaskOwner->GetOwnerActor();		
+		return TaskOwner->GetOwnerActor(this);		
 	}
 	else if (TasksComponent.IsValid())
 	{
-		return TasksComponent->GetOwner();
+		return TasksComponent->GetOwnerActor(this);
 	}
 
 	return nullptr;
@@ -113,11 +122,11 @@ AActor* UGameplayTask::GetAvatarActor() const
 {
 	if (TaskOwner.IsValid())
 	{
-		return TaskOwner->GetAvatarActor();
+		return TaskOwner->GetAvatarActor(this);
 	}
 	else if (TasksComponent.IsValid())
 	{
-		return TasksComponent->GetAvatarActor();
+		return TasksComponent->GetAvatarActor(this);
 	}
 
 	return nullptr;
@@ -194,13 +203,47 @@ FString UGameplayTask::GetDebugString() const
 	return FString::Printf(TEXT("Generic %s"), *GetName());
 }
 
-void UGameplayTask::RequireResource(TSubclassOf<UGameplayTaskResource> RequiredResource)
+void UGameplayTask::AddRequiredResource(TSubclassOf<UGameplayTaskResource> RequiredResource)
 {
 	check(RequiredResource);
 	const uint8 ResourceID = UGameplayTaskResource::GetResourceID(RequiredResource);
 	RequiredResources.AddID(ResourceID);	
 }
 
+void UGameplayTask::AddRequiredResourceSet(const TArray<TSubclassOf<UGameplayTaskResource> >& RequiredResourceSet)
+{
+	for (auto Resource : RequiredResourceSet)
+	{
+		const uint8 ResourceID = UGameplayTaskResource::GetResourceID(Resource);
+		RequiredResources.AddID(ResourceID);
+	}
+}
+
+void UGameplayTask::AddRequiredResourceSet(FGameplayResourceSet RequiredResourceSet)
+{
+	RequiredResources.AddSet(RequiredResourceSet);
+}
+
+void UGameplayTask::AddClaimedResource(TSubclassOf<UGameplayTaskResource> ClaimedResource)
+{
+	check(ClaimedResource);
+	const uint8 ResourceID = UGameplayTaskResource::GetResourceID(ClaimedResource);
+	ClaimedResources.AddID(ResourceID);
+}
+
+void UGameplayTask::AddClaimedResourceSet(const TArray<TSubclassOf<UGameplayTaskResource> >& AdditionalResourcesToClaim)
+{
+	for (auto ResourceClass : AdditionalResourcesToClaim)
+	{
+		ClaimedResources.AddID(UGameplayTaskResource::GetResourceID(ResourceClass));
+	}
+}
+
+void UGameplayTask::AddClaimedResourceSet(FGameplayResourceSet AdditionalResourcesToClaim)
+{
+	ClaimedResources.AddSet(AdditionalResourcesToClaim);
+}
+
 void UGameplayTask::PerformActivation()
 {
 	TaskState = EGameplayTaskState::Active;
diff --git a/Engine/Source/Runtime/GameplayTasks/Private/GameplayTaskOwnerInterface.cpp b/Engine/Source/Runtime/GameplayTasks/Private/GameplayTaskOwnerInterface.cpp
--- a/Engine/Source/Runtime/GameplayTasks/Private/GameplayTaskOwnerInterface.cpp
+++ b/Engine/Source/Runtime/GameplayTasks/Private/GameplayTaskOwnerInterface.cpp
@@ -8,7 +8,7 @@ void IGameplayTaskOwnerInterface::OnTaskInitialized(UGameplayTask& Task)
 
 }
 
-AActor* IGameplayTaskOwnerInterface::GetAvatarActor() const 
+AActor* IGameplayTaskOwnerInterface::GetAvatarActor(const UGameplayTask* Task) const
 { 
-	return GetOwnerActor(); 
-}
\ No newline at end of file
+	return GetOwnerActor(Task);
+}
diff --git a/Engine/Source/Runtime/GameplayTasks/Private/GameplayTasksComponent.cpp b/Engine/Source/Runtime/GameplayTasks/Private/GameplayTasksComponent.cpp
--- a/Engine/Source/Runtime/GameplayTasks/Private/GameplayTasksComponent.cpp
+++ b/Engine/Source/Runtime/GameplayTasks/Private/GameplayTasksComponent.cpp
@@ -3,6 +3,7 @@
 #include "GameplayTasksPrivatePCH.h"
 #include "GameplayTasksComponent.h"
 #include "GameplayTask.h"
+#include "MessageLog.h"
 
 #define LOCTEXT_NAMESPACE "GameplayTasksComponent"
 
@@ -177,7 +178,7 @@ void UGameplayTasksComponent::UpdateShouldTick()
 	}
 }
 
-AActor* UGameplayTasksComponent::GetAvatarActor() const
+AActor* UGameplayTasksComponent::GetAvatarActor(const UGameplayTask* Task) const
 {
 	return GetOwner();
 }
@@ -211,6 +212,25 @@ void UGameplayTasksComponent::RemoveResourceConsumingTask(UGameplayTask& Task)
 	}
 }
 
+void UGameplayTasksComponent::EndAllResourceConsumingTasksOwnedBy(const IGameplayTaskOwnerInterface& TaskOwner)// , bool bNotifyOwner)
+{
+	for (int32 TaskIndex = TaskPriorityQueue.Num() - 1; TaskIndex >= 0; --TaskIndex)
+	{
+		if (TaskPriorityQueue[TaskIndex] == nullptr)
+		{
+			TaskPriorityQueue.RemoveAt(TaskIndex, 1, /*bAllowShrinking=*/false);
+		}
+		else if (TaskPriorityQueue[TaskIndex]->GetTaskOwner() == &TaskOwner)
+		{
+			UGameplayTask* Task = TaskPriorityQueue[TaskIndex];
+			TaskPriorityQueue.RemoveAt(TaskIndex, 1, /*bAllowShrinking=*/false);
+			Task->TaskOwnerEnded();
+		}
+	}
+
+	UpdateTaskActivationFromIndex(0, FGameplayResourceSet(), FGameplayResourceSet());
+}
+
 void UGameplayTasksComponent::ProcessTaskEvents()
 {
 	static const int32 ErronousIterationsLimit = 100;
@@ -267,8 +287,8 @@ void UGameplayTasksComponent::AddTaskToPriorityQueue(UGameplayTask& NewTask)
 	
 	const bool bStartOnTopOfSamePriority = NewTask.GetResourceOverlapPolicy() == ETaskResourceOverlapPolicy::StartOnTop;
 
-	FGameplayResourceSet ResourcesUsedByActiveUpToX;
-	FGameplayResourceSet ResourcesRequiredUpToX;
+	FGameplayResourceSet ResourcesClaimedUpToX;
+	FGameplayResourceSet ResourcesBlockedUpToIndex;
 	int32 InsertionPoint = INDEX_NONE;
 	
 	if (TaskPriorityQueue.Num() > 0)
@@ -289,12 +309,12 @@ void UGameplayTasksComponent::AddTaskToPriorityQueue(UGameplayTask& NewTask)
 				break;
 			}
 
-			const FGameplayResourceSet TasksResources = TaskPriorityQueue[TaskIndex]->GetRequiredResources();
+			const FGameplayResourceSet ClaimedResources = TaskPriorityQueue[TaskIndex]->GetClaimedResources();
 			if (TaskPriorityQueue[TaskIndex]->IsActive())
 			{
-				ResourcesUsedByActiveUpToX.AddSet(TasksResources);
+				ResourcesClaimedUpToX.AddSet(ClaimedResources);
 			}
-			ResourcesRequiredUpToX.AddSet(TasksResources);
+			ResourcesBlockedUpToIndex.AddSet(ClaimedResources);
 		}
 	}
 	
@@ -304,84 +324,105 @@ void UGameplayTasksComponent::AddTaskToPriorityQueue(UGameplayTask& NewTask)
 		InsertionPoint = TaskPriorityQueue.Num() - 1;
 	}
 
-	UpdateTaskActivationFromIndex(InsertionPoint, ResourcesUsedByActiveUpToX, ResourcesRequiredUpToX);
+	UpdateTaskActivationFromIndex(InsertionPoint, ResourcesClaimedUpToX, ResourcesBlockedUpToIndex);
 }
 
 void UGameplayTasksComponent::RemoveTaskFromPriorityQueue(UGameplayTask& Task)
 {	
-	if (TaskPriorityQueue.Num() > 1)
+	const int32 RemovedTaskIndex = TaskPriorityQueue.Find(&Task);
+	if (RemovedTaskIndex != INDEX_NONE)
 	{
-		const int32 RemovedTaskIndex = TaskPriorityQueue.Find(&Task);
-		ensure(RemovedTaskIndex != INDEX_NONE);
-
-		// sum up resources up to TaskIndex
-		FGameplayResourceSet ResourcesUsedByActiveUpToX;
-		FGameplayResourceSet ResourcesRequiredUpToX;
-		for (int32 Index = 0; Index < RemovedTaskIndex; ++Index)
+		if (TaskPriorityQueue.Num() > 1)
 		{
-			ensure(TaskPriorityQueue[Index]);
-			if (TaskPriorityQueue[Index] == nullptr)
+			// sum up resources up to TaskIndex
+			FGameplayResourceSet ResourcesClaimedUpToX;
+			FGameplayResourceSet ResourcesBlockedUpToIndex;
+			for (int32 TaskIndex = 0; TaskIndex < RemovedTaskIndex; ++TaskIndex)
 			{
-				continue;
-			}
-			
-			const FGameplayResourceSet TasksResources = TaskPriorityQueue[Index]->GetRequiredResources();
-			if (TaskPriorityQueue[Index]->IsActive())
-			{
-				ResourcesUsedByActiveUpToX.AddSet(TasksResources);
+				ensure(TaskPriorityQueue[TaskIndex]);
+				if (TaskPriorityQueue[TaskIndex] == nullptr)
+				{
+					continue;
+				}
+
+				const FGameplayResourceSet ClaimedResources = TaskPriorityQueue[TaskIndex]->GetClaimedResources();
+				if (TaskPriorityQueue[TaskIndex]->IsActive())
+				{
+					ResourcesClaimedUpToX.AddSet(ClaimedResources);
+				}				
+				ResourcesBlockedUpToIndex.AddSet(ClaimedResources);
 			}
-			ResourcesRequiredUpToX.AddSet(TasksResources);
-		}
 
-		// don't forget to actually remove the task from the queue
-		TaskPriorityQueue.RemoveAt(RemovedTaskIndex, 1, /*bAllowShrinking=*/false);
+			// don't forget to actually remove the task from the queue
+			TaskPriorityQueue.RemoveAt(RemovedTaskIndex, 1, /*bAllowShrinking=*/false);
 
-		// if it wasn't the last item then proceed as usual
-		if (RemovedTaskIndex < TaskPriorityQueue.Num())
-		{
-			UpdateTaskActivationFromIndex(RemovedTaskIndex, ResourcesUsedByActiveUpToX, ResourcesRequiredUpToX);
+			// if it wasn't the last item then proceed as usual
+			if (RemovedTaskIndex < TaskPriorityQueue.Num())
+			{
+				UpdateTaskActivationFromIndex(RemovedTaskIndex, ResourcesClaimedUpToX, ResourcesBlockedUpToIndex);
+			}
+			else
+			{
+				// no need to do extra processing. This was the last task, so 
+				// ResourcesUsedByActiveUpToX is the CurrentlyUsedResources
+				SetCurrentlyClaimedResources(ResourcesClaimedUpToX);
+			}
 		}
 		else
 		{
-			// no need to do extra processing. This was the last task, so 
-			// ResourcesUsedByActiveUpToX is the CurrentlyUsedResources
-			CurrentlyUsedResources = ResourcesUsedByActiveUpToX;
+			TaskPriorityQueue.Pop(/*bAllowShrinking=*/false);
+			SetCurrentlyClaimedResources(FGameplayResourceSet());
 		}
 	}
 	else
 	{
-		TaskPriorityQueue.Pop(/*bAllowShrinking=*/false);
-		CurrentlyUsedResources.Clear();
+		// take a note and ignore
+		UE_VLOG(this, LogGameplayTasks, Verbose, TEXT("RemoveTaskFromPriorityQueue for %s called, but it's not in the queue. Might have been already removed"), *Task.GetName());
 	}
 }
 
-void UGameplayTasksComponent::UpdateTaskActivationFromIndex(int32 StartingIndex, FGameplayResourceSet ResourcesUsedByActiveUpToIndex, FGameplayResourceSet ResourcesRequiredUpToIndex)
+void UGameplayTasksComponent::UpdateTaskActivationFromIndex(int32 StartingIndex, FGameplayResourceSet ResourcesClaimedUpToIndex, FGameplayResourceSet ResourcesBlockedUpToIndex)
 {
-	check(TaskPriorityQueue.IsValidIndex(StartingIndex));
-
-	TArray<UGameplayTask*> TaskPriorityQueueCopy = TaskPriorityQueue;
-	for (int32 TaskIndex = StartingIndex; TaskIndex < TaskPriorityQueueCopy.Num(); ++TaskIndex)
+	if (TaskPriorityQueue.Num() > 0)
 	{
-		ensure(TaskPriorityQueueCopy[TaskIndex]);
-		if (TaskPriorityQueueCopy[TaskIndex] == nullptr)
-		{
-			continue;
-		}
+		check(TaskPriorityQueue.IsValidIndex(StartingIndex));
 
-		const FGameplayResourceSet TasksResources = TaskPriorityQueueCopy[TaskIndex]->GetRequiredResources();
-		if (TasksResources.GetOverlap(ResourcesUsedByActiveUpToIndex).IsEmpty())
+		TArray<UGameplayTask*> TaskPriorityQueueCopy = TaskPriorityQueue;
+		for (int32 TaskIndex = StartingIndex; TaskIndex < TaskPriorityQueueCopy.Num(); ++TaskIndex)
 		{
-			TaskPriorityQueueCopy[TaskIndex]->ActivateInTaskQueue();
-			ResourcesUsedByActiveUpToIndex.AddSet(TasksResources);
-		}
-		else
-		{
-			TaskPriorityQueueCopy[TaskIndex]->PauseInTaskQueue();
+			ensure(TaskPriorityQueueCopy[TaskIndex]);
+			if (TaskPriorityQueueCopy[TaskIndex] == nullptr)
+			{
+				continue;
+			}
+
+			const FGameplayResourceSet RequiredResources = TaskPriorityQueueCopy[TaskIndex]->GetRequiredResources();
+			const FGameplayResourceSet ClaimedResources = TaskPriorityQueueCopy[TaskIndex]->GetClaimedResources();
+			if (RequiredResources.GetOverlap(ResourcesBlockedUpToIndex).IsEmpty())
+			{
+				TaskPriorityQueueCopy[TaskIndex]->ActivateInTaskQueue();
+				ResourcesClaimedUpToIndex.AddSet(ClaimedResources);
+			}
+			else
+			{
+				TaskPriorityQueueCopy[TaskIndex]->PauseInTaskQueue();
+			}
+			ResourcesBlockedUpToIndex.AddSet(ClaimedResources);
 		}
-		ResourcesRequiredUpToIndex.AddSet(TasksResources);
 	}
+	
+	SetCurrentlyClaimedResources(ResourcesClaimedUpToIndex);
+}
 
-	CurrentlyUsedResources = ResourcesUsedByActiveUpToIndex;
+void UGameplayTasksComponent::SetCurrentlyClaimedResources(FGameplayResourceSet NewClaimedSet)
+{
+	if (CurrentlyClaimedResources != NewClaimedSet)
+	{
+		FGameplayResourceSet ReleasedResources = FGameplayResourceSet(CurrentlyClaimedResources).RemoveSet(NewClaimedSet);
+		FGameplayResourceSet ClaimedResources = FGameplayResourceSet(NewClaimedSet).RemoveSet(CurrentlyClaimedResources);
+		CurrentlyClaimedResources = NewClaimedSet;
+		OnClaimedResourcesChange.Broadcast(ClaimedResources, ReleasedResources);
+	}
 }
 
 //----------------------------------------------------------------------//
@@ -440,6 +481,110 @@ FString UGameplayTasksComponent::GetTaskStateName(EGameplayTaskState Value)
 }
 #endif // ENABLE_VISUAL_LOG
 
+EGameplayTaskRunResult UGameplayTasksComponent::RunGameplayTask(IGameplayTaskOwnerInterface& TaskOwner, UGameplayTask& Task, uint8 Priority, FGameplayResourceSet AdditionalRequiredResources, FGameplayResourceSet AdditionalClaimedResources)
+{
+	const FText NoneText = FText::FromString(TEXT("None"));
+
+	if (Task.GetState() == EGameplayTaskState::Paused || Task.GetState() == EGameplayTaskState::Active)
+	{
+		// return as success if already running for the same owner, failure otherwise 
+		return Task.GetTaskOwner() == &TaskOwner
+			? (Task.GetState() == EGameplayTaskState::Paused ? EGameplayTaskRunResult::Success_Paused : EGameplayTaskRunResult::Success_Active)
+			: EGameplayTaskRunResult::Error;
+	}
+
+	// this is a valid situation if the task has been created via "Construct Object" mechanics
+	if (Task.GetState() == EGameplayTaskState::Uninitialized)
+	{
+		Task.InitTask(TaskOwner, Priority);
+	}
+
+	Task.AddRequiredResourceSet(AdditionalRequiredResources);
+	Task.AddClaimedResourceSet(AdditionalClaimedResources);
+	Task.ReadyForActivation();
+
+	switch (Task.GetState())
+	{
+	case EGameplayTaskState::AwaitingActivation:
+	case EGameplayTaskState::Paused:
+		return EGameplayTaskRunResult::Success_Paused;
+		break;
+	case EGameplayTaskState::Active:
+		return EGameplayTaskRunResult::Success_Active;
+		break;
+	case EGameplayTaskState::Finished:
+		return EGameplayTaskRunResult::Success_Active;
+		break;
+	}
+
+	return EGameplayTaskRunResult::Error;
+}
+
+//----------------------------------------------------------------------//
+// BP API
+//----------------------------------------------------------------------//
+EGameplayTaskRunResult UGameplayTasksComponent::K2_RunGameplayTask(TScriptInterface<IGameplayTaskOwnerInterface> TaskOwner, UGameplayTask* Task, uint8 Priority, TArray<TSubclassOf<UGameplayTaskResource> > AdditionalRequiredResources, TArray<TSubclassOf<UGameplayTaskResource> > AdditionalClaimedResources)
+{
+	const FText NoneText = FText::FromString(TEXT("None"));
+
+	if (TaskOwner.GetInterface() == nullptr)
+	{
+		FMessageLog("PIE").Error(FText::Format(
+			LOCTEXT("RunGameplayTaskNullOwner", "Tried running a gameplay task {0} while owner is None!"),
+			Task ? FText::FromName(Task->GetFName()) : NoneText));
+		return EGameplayTaskRunResult::Error;
+	}
+
+	IGameplayTaskOwnerInterface& OwnerInstance = *TaskOwner;
+
+	if (Task == nullptr)
+	{
+		FMessageLog("PIE").Error(FText::Format(
+			LOCTEXT("RunNullGameplayTask", "Tried running a None task for {0}"),
+			FText::FromString(Cast<UObject>(&OwnerInstance)->GetName())
+			));
+		return EGameplayTaskRunResult::Error;
+	}
+
+	if (Task->GetState() == EGameplayTaskState::Paused || Task->GetState() == EGameplayTaskState::Active)
+	{
+		FMessageLog("PIE").Warning(FText::Format(
+			LOCTEXT("RunNullGameplayTask", "Tried running a None task for {0}"),
+			FText::FromString(Cast<UObject>(&OwnerInstance)->GetName())
+			));
+		// return as success if already running for the same owner, failure otherwise 
+		return Task->GetTaskOwner() == &OwnerInstance 
+			? (Task->GetState() == EGameplayTaskState::Paused ? EGameplayTaskRunResult::Success_Paused : EGameplayTaskRunResult::Success_Active)
+			: EGameplayTaskRunResult::Error;
+	}
+
+	// this is a valid situation if the task has been created via "Construct Object" mechanics
+	if (Task->GetState() == EGameplayTaskState::Uninitialized)
+	{
+		Task->InitTask(OwnerInstance, Priority);
+	}
+
+	Task->AddRequiredResourceSet(AdditionalRequiredResources);
+	Task->AddClaimedResourceSet(AdditionalClaimedResources);
+	Task->ReadyForActivation();
+
+	switch (Task->GetState())
+	{
+	case EGameplayTaskState::AwaitingActivation:
+	case EGameplayTaskState::Paused:
+		return EGameplayTaskRunResult::Success_Paused;
+		break;
+	case EGameplayTaskState::Active:
+		return EGameplayTaskRunResult::Success_Active;
+		break;
+	case EGameplayTaskState::Finished:
+		return EGameplayTaskRunResult::Success_Active;
+		break;
+	}
+
+	return EGameplayTaskRunResult::Error;
+}
+
 //----------------------------------------------------------------------//
 // FGameplayResourceSet
 //----------------------------------------------------------------------//
diff --git a/Engine/Source/Runtime/GameplayTasks/Private/Tasks/GameplayTask_SpawnActor.cpp b/Engine/Source/Runtime/GameplayTasks/Private/Tasks/GameplayTask_SpawnActor.cpp
--- a/Engine/Source/Runtime/GameplayTasks/Private/Tasks/GameplayTask_SpawnActor.cpp
+++ b/Engine/Source/Runtime/GameplayTasks/Private/Tasks/GameplayTask_SpawnActor.cpp
@@ -13,7 +13,7 @@ UGameplayTask_SpawnActor* UGameplayTask_SpawnActor::SpawnActor(TScriptInterface<
 	bool bCanSpawn = true;
 	if (bSpawnOnlyOnAuthority == true)
 	{
-		AActor* TaskOwnerActor = TaskOwner->GetOwnerActor();
+		AActor* TaskOwnerActor = TaskOwner->GetOwnerActor(nullptr);
 		if (TaskOwnerActor)
 		{
 			bCanSpawn = (TaskOwnerActor->Role == ROLE_Authority);
diff --git a/Engine/Source/Runtime/GameplayTasks/Public/GameplayTaskTypes.h b/Engine/Source/Runtime/GameplayTasks/Public/GameplayTaskTypes.h
--- a/Engine/Source/Runtime/GameplayTasks/Public/GameplayTaskTypes.h
+++ b/Engine/Source/Runtime/GameplayTasks/Public/GameplayTaskTypes.h
@@ -49,90 +49,3 @@ private:
 	TWeakObjectPtr<UObject> ObjectInstance;
 	TInterface* InterfaceInstance;
 };
-
-struct GAMEPLAYTASKS_API FGameplayResourceSet
-{
-	typedef uint16 FFlagContainer;
-	typedef uint8 FResourceID;
-
-	enum
-	{
-		MaxResources = sizeof(FFlagContainer)* 8
-	};
-
-private:
-	FFlagContainer Flags;
-
-public:
-	/** Mind that this constructor takes _flags_ not individual IDs */
-	explicit FGameplayResourceSet(FFlagContainer InFlags = 0) : Flags(InFlags)
-	{}
-	
-	FFlagContainer GetFlags() const 
-	{ 
-		return Flags; 
-	}
-	bool IsEmpty() const
-	{
-		return Flags == 0;
-	}
-	FGameplayResourceSet& AddID(uint8 ResourceID)
-	{
-		ensure(ResourceID < MaxResources);
-		Flags |= (1 << ResourceID);
-		return *this;
-	}
-	FGameplayResourceSet& RemoveID(uint8 ResourceID)
-	{
-		ensure(ResourceID < MaxResources);
-		Flags &= ~(1 << ResourceID);
-		return *this;
-	}
-	bool HasID(uint8 ResourceID) const
-	{
-		ensure(ResourceID < MaxResources);
-		return (Flags & (1 << ResourceID)) != 0;
-	}
-	FGameplayResourceSet& AddSet(FGameplayResourceSet Other)
-	{
-		Flags |= Other.Flags;
-		return *this;
-	}
-	void Clear()
-	{
-		Flags = FFlagContainer(0);
-	}
-	bool HasAllIDs(FGameplayResourceSet Other) const
-	{
-		return (Flags & Other.Flags) == Flags;
-	}
-	bool HasAnyID(FGameplayResourceSet Other) const
-	{
-		return (Flags & Other.Flags) != 0;
-	}
-	FGameplayResourceSet GetOverlap(FGameplayResourceSet Other) const
-	{
-		return FGameplayResourceSet(Flags & Other.Flags);
-	}
-	FGameplayResourceSet GetDifference(FGameplayResourceSet Other) const
-	{
-		return FGameplayResourceSet(Flags & ~(Flags & Other.Flags));
-	}
-
-	bool operator==(const FGameplayResourceSet& Other) const
-	{
-		return Flags == Other.Flags;
-	}
-
-	static FGameplayResourceSet AllResources()
-	{
-		return FGameplayResourceSet(FFlagContainer(-1));
-	}
-
-	static FGameplayResourceSet NoResources()
-	{
-		return FGameplayResourceSet(FFlagContainer(0));
-	}
-
-	FString GetDebugDescription() const;
-};
\ No newline at end of file
