A file missing from previous submission #UE4

It missing just causes some GameplayTask tests to fail

diff --git a/Engine/Source/Developer/AITestSuite/Private/MockAI/MockGameplayTasks.cpp b/Engine/Source/Developer/AITestSuite/Private/MockAI/MockGameplayTasks.cpp
--- a/Engine/Source/Developer/AITestSuite/Private/MockAI/MockGameplayTasks.cpp
+++ b/Engine/Source/Developer/AITestSuite/Private/MockAI/MockGameplayTasks.cpp
@@ -19,6 +19,7 @@ UMockTask_Log* UMockTask_Log::CreateTask(IGameplayTaskOwnerInterface& TaskOwner,
 	{
 		Task->Logger = &InLogger;
 		Task->RequiredResources = Resources;
+		Task->ClaimedResources = Resources;
 		Task->Priority = Priority;
 	}
 	return Task;
