Added a unity test for a Gameplay Tasks resource handling bug I've found yesterday #UE4

Merged over from Orion

diff --git a/Engine/Source/Developer/AITestSuite/Classes/MockGameplayTasks.h b/Engine/Source/Developer/AITestSuite/Classes/MockGameplayTasks.h
--- a/Engine/Source/Developer/AITestSuite/Classes/MockGameplayTasks.h
+++ b/Engine/Source/Developer/AITestSuite/Classes/MockGameplayTasks.h
@@ -25,6 +25,7 @@ class UMockTask_Log : public UGameplayTask
 
 protected:
 	FTestLogger<int32>* Logger;
+	bool bShoudEndAsPartOfActivation;
 
 public:
 	UMockTask_Log(const FObjectInitializer& ObjectInitializer);
@@ -42,6 +43,8 @@ public:
 
 	// testing only hack-functions
 	void EnableTick() { bTickingTask = true; }
+
+	void SetInstaEnd(bool bNewValue) { bShoudEndAsPartOfActivation = bNewValue; }
 };
 
 //
diff --git a/Engine/Source/Developer/AITestSuite/Private/MockAI/MockGameplayTasks.cpp b/Engine/Source/Developer/AITestSuite/Private/MockAI/MockGameplayTasks.cpp
--- a/Engine/Source/Developer/AITestSuite/Private/MockAI/MockGameplayTasks.cpp
+++ b/Engine/Source/Developer/AITestSuite/Private/MockAI/MockGameplayTasks.cpp
@@ -8,6 +8,7 @@
 UMockTask_Log::UMockTask_Log(const FObjectInitializer& ObjectInitializer)
 	: Super(ObjectInitializer)
 	, Logger(nullptr)
+	, bShoudEndAsPartOfActivation(false)
 {
 
 }
@@ -29,6 +30,10 @@ void UMockTask_Log::Activate()
 {
 	Logger->Log(ETestTaskMessage::Activate);
 	Super::Activate();
+	if (bShoudEndAsPartOfActivation)
+	{
+		EndTask();
+	}
 }
 
 void UMockTask_Log::OnDestroy(bool bOwnerFinished)
diff --git a/Engine/Source/Developer/AITestSuite/Private/Tests/GameplayTasksTest.cpp b/Engine/Source/Developer/AITestSuite/Private/Tests/GameplayTasksTest.cpp
--- a/Engine/Source/Developer/AITestSuite/Private/Tests/GameplayTasksTest.cpp
+++ b/Engine/Source/Developer/AITestSuite/Private/Tests/GameplayTasksTest.cpp
@@ -468,6 +468,28 @@ struct FAITest_GameplayTask_ClaimedResources : public FAITest_GameplayTasksTest
 };
 IMPLEMENT_AI_INSTANT_TEST(FAITest_GameplayTask_ClaimedResources, "System.Engine.AI.Gameplay Tasks.Claimed resources")
 
+//----------------------------------------------------------------------//
+// 
+//----------------------------------------------------------------------//
+struct FAITest_GameplayTask_ClaimedResourcesAndInstantFinish : public FAITest_GameplayTasksTest
+{
+	static const int32 TasksCount = 4;
+	UMockTask_Log* Task;
+
+	void InstantTest()
+	{
+		UWorld& World = GetWorld();
+
+		Task = UMockTask_Log::CreateTask(*Component, Logger, MovementResourceSet);
+		Task->SetInstaEnd(true);
+		Task->ReadyForActivation();
+
+		Test(TEXT("No claimed resources should be left behind"), Component->GetCurrentlyUsedResources().IsEmpty() == true);
+		Test(TEXT("There should no active tasks when task auto-insta-ended"), Component->GetTaskPriorityQueueSize() == 0);
+	}
+};
+IMPLEMENT_AI_INSTANT_TEST(FAITest_GameplayTask_ClaimedResourcesAndInstantFinish, "System.Engine.AI.Gameplay Tasks.Claimed resources vs Insta-finish tasks")
+
 // add tests if component wants ticking at while aborting/reactivating tasks
 // add test for re-adding/re-activating a finished task
 
