Fixed GC-related crashes in AITesting framework due to using unreferenced BT assets #UE4

diff --git a/Engine/Source/Developer/AITestSuite/Private/Tests/AITestsCommon.cpp b/Engine/Source/Developer/AITestSuite/Private/Tests/AITestsCommon.cpp
--- a/Engine/Source/Developer/AITestSuite/Private/Tests/AITestsCommon.cpp
+++ b/Engine/Source/Developer/AITestSuite/Private/Tests/AITestsCommon.cpp
@@ -122,6 +122,10 @@ FAITest_SimpleBT::FAITest_SimpleBT()
 	bUseSystemTicking = false;
 	
 	BTAsset = &FBTBuilder::CreateBehaviorTree();
+	if (BTAsset)
+	{
+		AddAutoDestroyObject(*BTAsset);
+	}
 }
 
 void FAITest_SimpleBT::SetUp()
diff --git a/Engine/Source/Developer/AITestSuite/Private/Tests/BTTest.cpp b/Engine/Source/Developer/AITestSuite/Private/Tests/BTTest.cpp
--- a/Engine/Source/Developer/AITestSuite/Private/Tests/BTTest.cpp
+++ b/Engine/Source/Developer/AITestSuite/Private/Tests/BTTest.cpp
@@ -771,7 +771,9 @@ struct FAITest_BTSubtreeSimple : public FAITest_SimpleBT
 	FAITest_BTSubtreeSimple()
 	{
 		UBehaviorTree* ChildAsset1 = &FBTBuilder::CreateBehaviorTree(*BTAsset);
+		if (ChildAsset1)
 		{
+			AddAutoDestroyObject(*ChildAsset1);
 			UBTCompositeNode& CompNode = FBTBuilder::AddSequence(*ChildAsset1);
 			{
 				FBTBuilder::AddTask(CompNode, 10, EBTNodeResult::Succeeded);
@@ -780,7 +782,9 @@ struct FAITest_BTSubtreeSimple : public FAITest_SimpleBT
 		}
 
 		UBehaviorTree* ChildAsset2 = &FBTBuilder::CreateBehaviorTree(*BTAsset);
+		if (ChildAsset2)
 		{
+			AddAutoDestroyObject(*ChildAsset2);
 			UBTCompositeNode& CompNode = FBTBuilder::AddSequence(*ChildAsset2);
 			{
 				FBTBuilder::AddTask(CompNode, 20, EBTNodeResult::Failed);
@@ -821,7 +825,9 @@ struct FAITest_BTSubtreeAbortOut : public FAITest_SimpleBT
 	FAITest_BTSubtreeAbortOut()
 	{
 		UBehaviorTree* ChildAsset = &FBTBuilder::CreateBehaviorTree(*BTAsset);
+		if (ChildAsset)
 		{
+			AddAutoDestroyObject(*ChildAsset);
 			UBTCompositeNode& CompNode = FBTBuilder::AddSequence(*ChildAsset);
 			{
 				FBTBuilder::AddTask(CompNode, 10, EBTNodeResult::Succeeded);
diff --git a/Engine/Source/Developer/AITestSuite/Public/AITestsCommon.h b/Engine/Source/Developer/AITestSuite/Public/AITestsCommon.h
--- a/Engine/Source/Developer/AITestSuite/Public/AITestsCommon.h
+++ b/Engine/Source/Developer/AITestSuite/Public/AITestsCommon.h
@@ -42,7 +42,7 @@ struct AITESTSUITE_API FAITestBase
 {
 private:
 	// internals
-	TArray<UObject*> SpawnedObjects;	
+	TArray<UObject*> SpawnedObjects;
 	uint32 bTearedDown : 1;
 protected:
 	FAutomationTestBase* TestRunner;
@@ -59,6 +59,12 @@ protected:
 		return ObjectInstance;
 	}
 
+	void AddAutoDestroyObject(UObject& ObjectRef)
+	{
+		ObjectRef.AddToRoot();
+		SpawnedObjects.Add(&ObjectRef);
+	}
+
 	UWorld& GetWorld() const
 	{
 		UWorld* World = FAITestHelpers::GetWorld();
