More Pipe unit testing

Successful reflight at https://horde.devtools.epicgames.com/job/6048bcf0ecdf57000140ae12 not submitted due to "internal error". May be due to missing "okforgithub" tag. Submitting manually.

diff --git a/Engine/Source/Developer/AITestSuite/Public/AITestsCommon.h b/Engine/Source/Developer/AITestSuite/Public/AITestsCommon.h
--- a/Engine/Source/Developer/AITestSuite/Public/AITestsCommon.h
+++ b/Engine/Source/Developer/AITestSuite/Public/AITestsCommon.h
@@ -28,7 +28,7 @@ private:
 
 namespace FAITestHelpers
 {
-	UWorld* GetWorld();
+	AITESTSUITE_API UWorld* GetWorld();
 	static const float TickInterval = 1.f / 30;
 
 	void UpdateFrameCounter();
@@ -111,9 +111,9 @@ DEFINE_EXPORTED_LATENT_AUTOMATION_COMMAND_ONE_PARAMETER(AITESTSUITE_API, FAITest
 		{ \
 			/* call the instant-test code */ \
 			bSuccess = TestInstance->InstantTest(); \
-			/* tear down */ \
-			TestInstance->TearDown(); \
 		}\
+		/* tear down */ \
+		TestInstance->TearDown(); \
 		delete TestInstance; \
 		return bSuccess; \
 	} 
