[AITestSuite] Added actor to the MockAI for VLOG support

diff --git a/Engine/Source/Developer/AITestSuite/Classes/MockAI.h b/Engine/Source/Developer/AITestSuite/Classes/MockAI.h
--- a/Engine/Source/Developer/AITestSuite/Classes/MockAI.h
+++ b/Engine/Source/Developer/AITestSuite/Classes/MockAI.h
@@ -4,7 +4,6 @@
 
 #include "CoreMinimal.h"
 #include "Stats/Stats.h"
-#include "UObject/ObjectMacros.h"
 #include "UObject/Object.h"
 #include "AITestsCommon.h"
 #include "Tickable.h"
@@ -21,11 +20,11 @@ struct FTestTickHelper : FTickableGameObject
 {
 	TWeakObjectPtr<class UMockAI> Owner;
 
-	FTestTickHelper() : Owner(NULL) {}
-	virtual void Tick(float DeltaTime);
-	virtual bool IsTickable() const { return Owner.IsValid(); }
-	virtual bool IsTickableInEditor() const { return true; }
-	virtual TStatId GetStatId() const;
+	FTestTickHelper() : Owner(nullptr) {}
+	virtual void Tick(float DeltaTime) override;
+	virtual bool IsTickable() const override { return Owner.IsValid(); }
+	virtual bool IsTickableInEditor() const override { return true; }
+	virtual TStatId GetStatId() const override;
 };
 
 UCLASS()
@@ -33,26 +32,29 @@ class UMockAI : public UObject
 {
 	GENERATED_UCLASS_BODY()
 
-	virtual ~UMockAI();
+	virtual ~UMockAI() override;
 
 	FTestTickHelper TickHelper;
 
 	UPROPERTY()
-	TObjectPtr<UBlackboardComponent> BBComp;
+	TObjectPtr<AActor> Actor = nullptr;
 
 	UPROPERTY()
-	TObjectPtr<UBrainComponent> BrainComp;
+	TObjectPtr<UBlackboardComponent> BBComp = nullptr;
 
 	UPROPERTY()
-	TObjectPtr<UAIPerceptionComponent> PerceptionComp;
+	TObjectPtr<UBrainComponent> BrainComp = nullptr;
 
 	UPROPERTY()
-	TObjectPtr<UPawnActionsComponent> PawnActionComp;
+	TObjectPtr<UAIPerceptionComponent> PerceptionComp = nullptr;
+
+	UPROPERTY()
+	TObjectPtr<UPawnActionsComponent> PawnActionComp = nullptr;
 	
 	template<typename TBrainClass>
 	void UseBrainComponent()
 	{
-		BrainComp = NewObject<TBrainClass>(FAITestHelpers::GetWorld());
+		BrainComp = NewObject<TBrainClass>(Actor);
 	}
 
 	void UseBlackboardComponent();
diff --git a/Engine/Source/Developer/AITestSuite/Private/MockAI/MockAI.cpp b/Engine/Source/Developer/AITestSuite/Private/MockAI/MockAI.cpp
--- a/Engine/Source/Developer/AITestSuite/Private/MockAI/MockAI.cpp
+++ b/Engine/Source/Developer/AITestSuite/Private/MockAI/MockAI.cpp
@@ -28,11 +28,28 @@ TStatId FTestTickHelper::GetStatId() const
 UMockAI::UMockAI(const FObjectInitializer& ObjectInitializer)
 	: Super(ObjectInitializer)
 {
+	if (!HasAnyFlags(RF_ClassDefaultObject))
+	{
+		UWorld* World = FAITestHelpers::GetWorld();
+		if (ensureMsgf(World != nullptr, TEXT("A world is required to spawn the associated test actor")))
+		{
+			FActorSpawnParameters SpawnParameters;
+			SpawnParameters.Name = TEXT("MockAIActor");
+			SpawnParameters.NameMode = FActorSpawnParameters::ESpawnActorNameMode::Requested;
+			SpawnParameters.ObjectFlags = RF_Transient;
+			Actor = World->SpawnActor<AActor>(SpawnParameters);
+		}
+	}
 }
 
 UMockAI::~UMockAI()
 {
 	TickHelper.Owner.Reset();
+
+	if (Actor != nullptr)
+	{
+		Actor->Destroy();
+	}
 }
 
 void UMockAI::SetEnableTicking(bool bShouldTick)
@@ -43,45 +60,45 @@ void UMockAI::SetEnableTicking(bool bShouldTick)
 	}
 	else
 	{
-		TickHelper.Owner = NULL;
+		TickHelper.Owner = nullptr;
 	}
 }
 
 void UMockAI::UseBlackboardComponent()
 {
-	BBComp = NewObject<UBlackboardComponent>(FAITestHelpers::GetWorld());
+	BBComp = NewObject<UBlackboardComponent>(Actor);
 }
 
 void UMockAI::UsePerceptionComponent()
 {
-	PerceptionComp = NewObject<UAIPerceptionComponent>(FAITestHelpers::GetWorld());
+	PerceptionComp = NewObject<UAIPerceptionComponent>(Actor);
 }
 
 void UMockAI::UsePawnActionsComponent()
 {
-	PawnActionComp = NewObject<UPawnActionsComponent>(FAITestHelpers::GetWorld());
+	PawnActionComp = NewObject<UPawnActionsComponent>(Actor);
 }
 
 void UMockAI::TickMe(float DeltaTime)
 {
 	if (BBComp)
 	{
-		BBComp->TickComponent(DeltaTime, ELevelTick::LEVELTICK_All, NULL);
+		BBComp->TickComponent(DeltaTime, ELevelTick::LEVELTICK_All, nullptr);
 	}
 
 	if (PerceptionComp)
 	{
-		PerceptionComp->TickComponent(DeltaTime, ELevelTick::LEVELTICK_All, NULL);
+		PerceptionComp->TickComponent(DeltaTime, ELevelTick::LEVELTICK_All, nullptr);
 	}
 	
 	if (BrainComp)
 	{
-		BrainComp->TickComponent(DeltaTime, ELevelTick::LEVELTICK_All, NULL);
+		BrainComp->TickComponent(DeltaTime, ELevelTick::LEVELTICK_All, nullptr);
 	}
 	
 	if (PawnActionComp)
 	{
-		PawnActionComp->TickComponent(DeltaTime, ELevelTick::LEVELTICK_All, NULL);
+		PawnActionComp->TickComponent(DeltaTime, ELevelTick::LEVELTICK_All, nullptr);
 	}
 }
 
