[FAITestHelpers] added optional outer parameter to NewAutoDestroyObject
#rnx

diff --git a/Engine/Source/Developer/AITestSuite/Public/AITestsCommon.h b/Engine/Source/Developer/AITestSuite/Public/AITestsCommon.h
--- a/Engine/Source/Developer/AITestSuite/Public/AITestsCommon.h
+++ b/Engine/Source/Developer/AITestSuite/Public/AITestsCommon.h
@@ -48,9 +48,9 @@ protected:
 	{}
 
 	template<typename ClassToSpawn>
-	ClassToSpawn* NewAutoDestroyObject()
+	ClassToSpawn* NewAutoDestroyObject(UObject* Outer = GetTransientPackage())
 	{
-		ClassToSpawn* ObjectInstance = NewObject<ClassToSpawn>();
+		ClassToSpawn* ObjectInstance = NewObject<ClassToSpawn>(Outer);
 		ObjectInstance->AddToRoot();
 		SpawnedObjects.Add(ObjectInstance);
 		return ObjectInstance;
@@ -170,7 +170,7 @@ struct FAITest_SimpleComponentBasedTest : public FAITestBase
 		Logger.TestRunner = TestRunner;
 	}
 
-	virtual ~FAITest_SimpleComponentBasedTest()
+	virtual ~FAITest_SimpleComponentBasedTest() override
 	{
 		GetTestRunner().TestTrue(TEXT("Not all expected values has been logged"), Logger.ExpectedValues.Num() == 0 || Logger.ExpectedValues.Num() == Logger.LoggedValues.Num());
 	}
