Added a bunch of new unit test features to test new behavior tree code path.

#rnx

diff --git a/Engine/Source/Developer/AITestSuite/Classes/BehaviorTree/TestBTService_Log.h b/Engine/Source/Developer/AITestSuite/Classes/BehaviorTree/TestBTService_Log.h
--- a/Engine/Source/Developer/AITestSuite/Classes/BehaviorTree/TestBTService_Log.h
+++ b/Engine/Source/Developer/AITestSuite/Classes/BehaviorTree/TestBTService_Log.h
@@ -20,6 +20,12 @@ class UTestBTService_Log : public UBTService
 	UPROPERTY()
 	FName KeyNameTick;
 
+	UPROPERTY()
+	FName KeyNameBecomeRelevant;
+
+	UPROPERTY()
+	FName KeyNameCeaseRelevant;
+
 	UPROPERTY()
 	int32 LogTick;
 
diff --git a/Engine/Source/Developer/AITestSuite/Classes/BehaviorTree/TestBTTask_Log.h b/Engine/Source/Developer/AITestSuite/Classes/BehaviorTree/TestBTTask_Log.h
--- a/Engine/Source/Developer/AITestSuite/Classes/BehaviorTree/TestBTTask_Log.h
+++ b/Engine/Source/Developer/AITestSuite/Classes/BehaviorTree/TestBTTask_Log.h
@@ -25,6 +25,9 @@ class UTestBTTask_Log : public UBTTaskNode
 	UPROPERTY()
 	int32 ExecutionTicks;
 
+	UPROPERTY()
+	int32 LogTickIndex;
+
 	UPROPERTY()
 	TEnumAsByte<EBTNodeResult::Type> LogResult;
 
diff --git a/Engine/Source/Developer/AITestSuite/Private/BehaviorTree/TestBTService_Log.cpp b/Engine/Source/Developer/AITestSuite/Private/BehaviorTree/TestBTService_Log.cpp
--- a/Engine/Source/Developer/AITestSuite/Private/BehaviorTree/TestBTService_Log.cpp
+++ b/Engine/Source/Developer/AITestSuite/Private/BehaviorTree/TestBTService_Log.cpp
@@ -15,12 +15,23 @@ UTestBTService_Log::UTestBTService_Log(const FObjectInitializer& ObjectInitializ
 	LogDeactivation = INDEX_NONE;
 	KeyNameTick = NAME_None;
 	LogTick = INDEX_NONE;
+	KeyNameBecomeRelevant = NAME_None;
+	KeyNameCeaseRelevant = NAME_None;
+
+	// Force the service to tick every frame
+	Interval = 0.0f;
+	RandomDeviation = 0.0f;
 }
 
 void UTestBTService_Log::OnBecomeRelevant(UBehaviorTreeComponent& OwnerComp, uint8* NodeMemory)
 {
 	Super::OnBecomeRelevant(OwnerComp, NodeMemory);
 
+	if (KeyNameBecomeRelevant != NAME_None)
+	{
+		OwnerComp.GetBlackboardComponent()->SetValueAsBool(KeyNameBecomeRelevant, true);
+	}
+
 	if (LogActivation >= 0)
 	{
 		UMockAI_BT::ExecutionLog.Add(LogActivation);
@@ -31,6 +42,11 @@ void UTestBTService_Log::OnCeaseRelevant(UBehaviorTreeComponent& OwnerComp, uint
 {
 	Super::OnCeaseRelevant(OwnerComp, NodeMemory);
 
+	if (KeyNameCeaseRelevant != NAME_None)
+	{
+		OwnerComp.GetBlackboardComponent()->SetValueAsBool(KeyNameCeaseRelevant, true);
+	}
+
 	if (LogDeactivation >= 0)
 	{
 		UMockAI_BT::ExecutionLog.Add(LogDeactivation);
diff --git a/Engine/Source/Developer/AITestSuite/Private/BehaviorTree/TestBTTask_Log.cpp b/Engine/Source/Developer/AITestSuite/Private/BehaviorTree/TestBTTask_Log.cpp
--- a/Engine/Source/Developer/AITestSuite/Private/BehaviorTree/TestBTTask_Log.cpp
+++ b/Engine/Source/Developer/AITestSuite/Private/BehaviorTree/TestBTTask_Log.cpp
@@ -10,6 +10,7 @@ UTestBTTask_Log::UTestBTTask_Log(const FObjectInitializer& ObjectInitializer) :
 	LogIndex = 0;
 	LogFinished = -1;
 	LogResult = EBTNodeResult::Succeeded;
+	LogTickIndex = -1;
 
 	bNotifyTick = true;
 }
@@ -32,6 +33,11 @@ void UTestBTTask_Log::TickTask(UBehaviorTreeComponent& OwnerComp, uint8* NodeMem
 {
 	FBTLogTaskMemory* MyMemory = (FBTLogTaskMemory*)NodeMemory;
 
+	if (LogTickIndex != -1)
+	{
+		LogExecution(OwnerComp, LogTickIndex);
+	}
+
 	if (FAITestHelpers::FramesCounter() >= MyMemory->EndFrameIdx)
 	{
 		LogExecution(OwnerComp, LogFinished);
diff --git a/Engine/Source/Developer/AITestSuite/Public/BTBuilder.h b/Engine/Source/Developer/AITestSuite/Public/BTBuilder.h
--- a/Engine/Source/Developer/AITestSuite/Public/BTBuilder.h
+++ b/Engine/Source/Developer/AITestSuite/Public/BTBuilder.h
@@ -125,12 +125,13 @@ struct FBTBuilder
 		return *NodeOb;
 	}
 
-	static void AddTask(UBTCompositeNode& ParentNode, int32 LogIndex, EBTNodeResult::Type NodeResult, int32 ExecutionTicks = 0)
+	static void AddTask(UBTCompositeNode& ParentNode, int32 LogIndex, EBTNodeResult::Type NodeResult, int32 ExecutionTicks = 0, int32 LogTickIndex = -1)
 	{
 		UTestBTTask_Log* TaskNode = NewObject<UTestBTTask_Log>(ParentNode.GetTreeAsset());
 		TaskNode->LogIndex = LogIndex;
 		TaskNode->LogResult = NodeResult;
 		TaskNode->ExecutionTicks = ExecutionTicks;
+		TaskNode->LogTickIndex = LogTickIndex;
 
 		const int32 ChildIdx = ParentNode.Children.AddZeroed(1);
 		ParentNode.Children[ChildIdx].ChildTask = TaskNode;
@@ -287,13 +288,15 @@ struct FBTBuilder
 		return *ServiceOb;
 	}
 
-	static void WithServiceLog(UBTCompositeNode& ParentNode, int32 ActivationIndex, int32 DeactivationIndex, int32 TickIndex = INDEX_NONE, FName BoolKeyName = NAME_None, bool bCallTickOnSearchStart = false)
+	static void WithServiceLog(UBTCompositeNode& ParentNode, int32 ActivationIndex, int32 DeactivationIndex, int32 TickIndex = INDEX_NONE, FName TickBoolKeyName = NAME_None, bool bCallTickOnSearchStart = false, FName BecomeRelevantBoolKeyName = NAME_None, FName CeaseRelevantBoolKeyName = NAME_None)
 	{
 		UTestBTService_Log& LogService = WithService<UTestBTService_Log>(ParentNode);
 		LogService.LogActivation = ActivationIndex;
 		LogService.LogDeactivation = DeactivationIndex;
 		LogService.LogTick = TickIndex;
-		LogService.SetFlagOnTick(BoolKeyName, bCallTickOnSearchStart);
+		LogService.SetFlagOnTick(TickBoolKeyName, bCallTickOnSearchStart);
+		LogService.KeyNameBecomeRelevant = BecomeRelevantBoolKeyName;
+		LogService.KeyNameCeaseRelevant = CeaseRelevantBoolKeyName;
 	}
 
 	template<class T>
@@ -308,12 +311,14 @@ struct FBTBuilder
 		return *ServiceOb;
 	}
 
-	static void WithTaskServiceLog(UBTCompositeNode& ParentNode, int32 ActivationIndex, int32 DeactivationIndex, int32 TickIndex = INDEX_NONE, FName BoolKeyName = NAME_None, bool bCallTickOnSearchStart = false)
+	static void WithTaskServiceLog(UBTCompositeNode& ParentNode, int32 ActivationIndex, int32 DeactivationIndex, int32 TickIndex = INDEX_NONE, FName TickBoolKeyName = NAME_None, bool bCallTickOnSearchStart = false, FName BecomeRelevantBoolKeyName = NAME_None, FName CeaseRelevantBoolKeyName = NAME_None)
 	{
 		UTestBTService_Log& LogService = WithTaskService<UTestBTService_Log>(ParentNode);
 		LogService.LogActivation = ActivationIndex;
 		LogService.LogDeactivation = DeactivationIndex;
 		LogService.LogTick = TickIndex;
-		LogService.SetFlagOnTick(BoolKeyName, bCallTickOnSearchStart);
+		LogService.SetFlagOnTick(TickBoolKeyName, bCallTickOnSearchStart);
+		LogService.KeyNameBecomeRelevant = BecomeRelevantBoolKeyName;
+		LogService.KeyNameCeaseRelevant = CeaseRelevantBoolKeyName;
 	}
 };
