Fix failing bt test looping forever

#rnx

diff --git a/Engine/Source/Developer/AITestSuite/Private/Tests/BTTest.cpp b/Engine/Source/Developer/AITestSuite/Private/Tests/BTTest.cpp
--- a/Engine/Source/Developer/AITestSuite/Private/Tests/BTTest.cpp
+++ b/Engine/Source/Developer/AITestSuite/Private/Tests/BTTest.cpp
@@ -65,7 +65,8 @@ struct FAITest_SimpleBT : public FAITestBase
 			}
 		}
 
-		return VerifyResults();
+		VerifyResults();
+		return true;
 	}
 
 	bool VerifyResults()
