Fix for BT hang from UDN case 00317509

diff --git a/Engine/Source/Developer/AITestSuite/Private/Tests/BTTest.cpp b/Engine/Source/Developer/AITestSuite/Private/Tests/BTTest.cpp
--- a/Engine/Source/Developer/AITestSuite/Private/Tests/BTTest.cpp
+++ b/Engine/Source/Developer/AITestSuite/Private/Tests/BTTest.cpp
@@ -422,6 +422,452 @@ struct FAITest_BTCompositeFailedDecoratorUnregistersChildren : public FAITest_Si
 };
 IMPLEMENT_AI_LATENT_TEST(FAITest_BTCompositeFailedDecoratorUnregistersChildren, "System.AI.Behavior Trees.Composite failed decorator unregisters child nodes")
 
+/* All BTAbortingDuringService/BTAbortingDuringTaskService come from UDN case 00317509*/
+struct FAITest_BTAbortingDuringServiceTick : public FAITest_SimpleBT
+{
+	FAITest_BTAbortingDuringServiceTick()
+	{
+		UBTCompositeNode& CompNode = FBTBuilder::AddSelector(*BTAsset);
+		{
+			UBTCompositeNode& CompNode2 = FBTBuilder::AddSequence(CompNode);
+			{
+				FBTBuilder::WithDecoratorBlackboard(CompNode, EBasicKeyOperation::NotSet, EBTFlowAbortMode::Self, TEXT("Bool1"));
+				FBTBuilder::WithServiceLog(CompNode2, 1/*ActivationIndex*/, 2/*DeactivationIndex*/, 3/*TickIndex*/, TEXT("Bool1"));
+
+				FBTBuilder::AddTask(CompNode2, 4, EBTNodeResult::Succeeded, 10/*ExecutionTicks*/, 222/*LogTickIndex*/);
+			}
+			FBTBuilder::AddTask(CompNode, 6, EBTNodeResult::Succeeded);
+		}
+
+		ExpectedResult.Add(1);
+		ExpectedResult.Add(3);
+		ExpectedResult.Add(4);
+		ExpectedResult.Add(222);
+		ExpectedResult.Add(3); // Extra service tick because aux nodes are always ticked before any pending requests are processed
+		ExpectedResult.Add(2);
+		ExpectedResult.Add(6);
+	}
+};
+IMPLEMENT_AI_LATENT_TEST(FAITest_BTAbortingDuringServiceTick, "System.AI.Behavior Trees.Abort: during service tick")
+
+struct FAITest_BTAbortingDuringServiceBecomeRelevant : public FAITest_SimpleBT
+{
+	FAITest_BTAbortingDuringServiceBecomeRelevant()
+	{
+		UBTCompositeNode& CompNode = FBTBuilder::AddSelector(*BTAsset);
+		{
+			UBTCompositeNode& CompNode2 = FBTBuilder::AddSequence(CompNode);
+			{
+				FBTBuilder::WithDecoratorBlackboard(CompNode, EBasicKeyOperation::NotSet, EBTFlowAbortMode::Self, TEXT("Bool1"));
+				FBTBuilder::WithServiceLog(CompNode2, 1/*ActivationIndex*/, 2/*DeactivationIndex*/, 3/*TickIndex*/, NAME_None, false, TEXT("Bool1"));
+
+				FBTBuilder::AddTask(CompNode2, 4, EBTNodeResult::Succeeded);
+			}
+			FBTBuilder::AddTask(CompNode, 6, EBTNodeResult::Succeeded);
+		}
+
+		ExpectedResult.Add(1);
+		ExpectedResult.Add(3);
+		ExpectedResult.Add(4);
+		ExpectedResult.Add(3); // Extra service tick because aux nodes are always ticked before any pending requests are processed
+		ExpectedResult.Add(2);
+		ExpectedResult.Add(6);
+	}
+};
+IMPLEMENT_AI_LATENT_TEST(FAITest_BTAbortingDuringServiceBecomeRelevant, "System.AI.Behavior Trees.Abort: during service becomes relevant")
+
+struct FAITest_BTAbortingDuringServiceBecomeRelevantTaskTicking : public FAITest_SimpleBT
+{
+	FAITest_BTAbortingDuringServiceBecomeRelevantTaskTicking()
+	{
+		UBTCompositeNode& CompNode = FBTBuilder::AddSelector(*BTAsset);
+		{
+			UBTCompositeNode& CompNode2 = FBTBuilder::AddSequence(CompNode);
+			{
+				FBTBuilder::WithDecoratorBlackboard(CompNode, EBasicKeyOperation::NotSet, EBTFlowAbortMode::Self, TEXT("Bool1"));
+				FBTBuilder::WithServiceLog(CompNode2, 1/*ActivationIndex*/, 2/*DeactivationIndex*/, 3/*TickIndex*/, NAME_None, false, TEXT("Bool1"));
+
+				FBTBuilder::AddTask(CompNode2, 4, EBTNodeResult::Succeeded, 10/*ExecutionTicks*/, 222/*LogTickIndex*/);
+			}
+			FBTBuilder::AddTask(CompNode, 6, EBTNodeResult::Succeeded);
+		}
+
+		ExpectedResult.Add(1);
+		ExpectedResult.Add(3);
+		ExpectedResult.Add(4);
+		ExpectedResult.Add(222);
+		ExpectedResult.Add(3); // Extra service tick because aux nodes are always ticked before any pending requests are processed
+		ExpectedResult.Add(2);
+		ExpectedResult.Add(6);
+	}
+};
+IMPLEMENT_AI_LATENT_TEST(FAITest_BTAbortingDuringServiceBecomeRelevantTaskTicking, "System.AI.Behavior Trees.Abort: during service becomes relevant when task continue to tick for a longer time")
+
+struct FAITest_BTAbortingDuringServiceCeaseRelevant : public FAITest_SimpleBT
+{
+	FAITest_BTAbortingDuringServiceCeaseRelevant()
+	{
+		UBTCompositeNode& CompNode = FBTBuilder::AddSequence(*BTAsset);
+		{
+			UBTCompositeNode& CompNode2 = FBTBuilder::AddSelector(CompNode);
+			{
+				UBTCompositeNode& CompNode3 = FBTBuilder::AddSelector(CompNode2);
+				{
+					FBTBuilder::WithDecoratorBlackboard(CompNode2, EBasicKeyOperation::NotSet, EBTFlowAbortMode::Self, TEXT("Bool1"));
+					FBTBuilder::WithServiceLog(CompNode3, 1/*ActivationIndex*/, 2/*DeactivationIndex*/, 3/*TickIndex*/, NAME_None, false, NAME_None, TEXT("Bool1"));
+
+					FBTBuilder::AddTask(CompNode3, 4, EBTNodeResult::Succeeded);
+				}
+				FBTBuilder::AddTask(CompNode2, 6, EBTNodeResult::Succeeded);
+			}
+			FBTBuilder::AddTask(CompNode, 7, EBTNodeResult::Succeeded);
+		}
+
+		ExpectedResult.Add(1);
+		ExpectedResult.Add(3);
+		ExpectedResult.Add(4);
+		ExpectedResult.Add(3); // Extra service tick because aux nodes are always ticked before any pending requests are processed
+		ExpectedResult.Add(2);
+		ExpectedResult.Add(7);
+	}
+};
+IMPLEMENT_AI_LATENT_TEST(FAITest_BTAbortingDuringServiceCeaseRelevant, "System.AI.Behavior Trees.Abort: during service ceases relevant")
+
+struct FAITest_BTAbortingDuringTaskServiceTick : public FAITest_SimpleBT
+{
+	FAITest_BTAbortingDuringTaskServiceTick()
+	{
+		UBTCompositeNode& CompNode = FBTBuilder::AddSelector(*BTAsset);
+		{
+			FBTBuilder::AddTask(CompNode, 4, EBTNodeResult::Succeeded, 10/*ExecutionTicks*/, 222/*LogTickIndex*/);
+			FBTBuilder::WithDecoratorBlackboard(CompNode, EBasicKeyOperation::NotSet, EBTFlowAbortMode::Self, TEXT("Bool1"));
+			FBTBuilder::WithTaskServiceLog(CompNode, 1/*ActivationIndex*/, 2/*DeactivationIndex*/, 3/*TickIndex*/, TEXT("Bool1"));
+
+			FBTBuilder::AddTask(CompNode, 6, EBTNodeResult::Succeeded);
+		}
+
+		ExpectedResult.Add(1);
+		ExpectedResult.Add(3);
+		ExpectedResult.Add(4);
+		ExpectedResult.Add(222);
+		ExpectedResult.Add(3); // Extra service tick because aux nodes are always ticked before any pending requests are processed
+		ExpectedResult.Add(2);
+		ExpectedResult.Add(6);
+	}
+};
+IMPLEMENT_AI_LATENT_TEST(FAITest_BTAbortingDuringTaskServiceTick, "System.AI.Behavior Trees.Abort: during task service tick")
+
+struct FAITest_BTAbortingDuringTaskServiceBecomeRelevant : public FAITest_SimpleBT
+{
+	FAITest_BTAbortingDuringTaskServiceBecomeRelevant()
+	{
+		UBTCompositeNode& CompNode = FBTBuilder::AddSelector(*BTAsset);
+		{
+			FBTBuilder::AddTask(CompNode, 4, EBTNodeResult::Succeeded);
+			FBTBuilder::WithDecoratorBlackboard(CompNode, EBasicKeyOperation::NotSet, EBTFlowAbortMode::Self, TEXT("Bool1"));
+			FBTBuilder::WithTaskServiceLog(CompNode, 1/*ActivationIndex*/, 2/*DeactivationIndex*/, 3/*TickIndex*/, NAME_None, false, TEXT("Bool1"));
+
+			FBTBuilder::AddTask(CompNode, 6, EBTNodeResult::Succeeded);
+		}
+
+		ExpectedResult.Add(1);
+		ExpectedResult.Add(3);
+		ExpectedResult.Add(4);
+		ExpectedResult.Add(3); // Extra service tick because aux nodes are always ticked before any pending requests are processed
+		ExpectedResult.Add(2);
+		ExpectedResult.Add(6);
+	}
+};
+IMPLEMENT_AI_LATENT_TEST(FAITest_BTAbortingDuringTaskServiceBecomeRelevant, "System.AI.Behavior Trees.Abort: during task service becomes relevant")
+
+struct FAITest_BTAbortingDuringTaskServiceBecomeRelevantTaskTicking : public FAITest_SimpleBT
+{
+	FAITest_BTAbortingDuringTaskServiceBecomeRelevantTaskTicking()
+	{
+		UBTCompositeNode& CompNode = FBTBuilder::AddSelector(*BTAsset);
+		{
+			FBTBuilder::AddTask(CompNode, 4, EBTNodeResult::Succeeded, 10/*ExecutionTicks*/, 222/*LogTickIndex*/);
+			FBTBuilder::WithDecoratorBlackboard(CompNode, EBasicKeyOperation::NotSet, EBTFlowAbortMode::Self, TEXT("Bool1"));
+			FBTBuilder::WithTaskServiceLog(CompNode, 1/*ActivationIndex*/, 2/*DeactivationIndex*/, 3/*TickIndex*/, NAME_None, false, TEXT("Bool1"));
+
+			FBTBuilder::AddTask(CompNode, 6, EBTNodeResult::Succeeded);
+		}
+
+		ExpectedResult.Add(1);
+		ExpectedResult.Add(3);
+		ExpectedResult.Add(4);
+		ExpectedResult.Add(222);
+		ExpectedResult.Add(3); // Extra service tick because aux nodes are always ticked before any pending requests are processed
+		ExpectedResult.Add(2);
+		ExpectedResult.Add(6);
+	}
+};
+IMPLEMENT_AI_LATENT_TEST(FAITest_BTAbortingDuringTaskServiceBecomeRelevantTaskTicking, "System.AI.Behavior Trees.Abort: during task service becomes relevant when task continue to tick for a longer time")
+
+struct FAITest_BTAbortingDuringTaskServiceCeaseRelevant : public FAITest_SimpleBT
+{
+	FAITest_BTAbortingDuringTaskServiceCeaseRelevant()
+	{
+		UBTCompositeNode& CompNode = FBTBuilder::AddSequence(*BTAsset);
+		{
+			UBTCompositeNode& CompNode2 = FBTBuilder::AddSelector(CompNode);
+			{
+				FBTBuilder::AddTask(CompNode2, 4, EBTNodeResult::Succeeded);
+				FBTBuilder::WithDecoratorBlackboard(CompNode2, EBasicKeyOperation::NotSet, EBTFlowAbortMode::Self, TEXT("Bool1"));
+				FBTBuilder::WithTaskServiceLog(CompNode2, 1/*ActivationIndex*/, 2/*DeactivationIndex*/, 3/*TickIndex*/, NAME_None, false, NAME_None, TEXT("Bool1"));
+
+				FBTBuilder::AddTask(CompNode2, 6, EBTNodeResult::Succeeded);
+			}
+			FBTBuilder::AddTask(CompNode, 7, EBTNodeResult::Succeeded);
+		}
+
+		ExpectedResult.Add(1);
+		ExpectedResult.Add(3);
+		ExpectedResult.Add(4);
+		ExpectedResult.Add(3); // Extra service tick because aux nodes are always ticked before any pending requests are processed
+		ExpectedResult.Add(2);
+		ExpectedResult.Add(7);
+	}
+};
+IMPLEMENT_AI_LATENT_TEST(FAITest_BTAbortingDuringTaskServiceCeaseRelevant, "System.AI.Behavior Trees.Abort: during task service ceases relevant")
+
+struct FAITest_BTSwitchingHigherPrioDuringServiceTick : public FAITest_SimpleBT
+{
+	FAITest_BTSwitchingHigherPrioDuringServiceTick()
+	{
+		UBTCompositeNode& CompNode = FBTBuilder::AddSelector(*BTAsset);
+		{
+			UBTCompositeNode& CompNode2 = FBTBuilder::AddSequence(CompNode);
+			{
+				FBTBuilder::WithDecoratorBlackboard(CompNode, EBasicKeyOperation::Set, EBTFlowAbortMode::LowerPriority, TEXT("Bool1"));
+
+				FBTBuilder::AddTask(CompNode2, 4, EBTNodeResult::Succeeded);
+			}
+			UBTCompositeNode& CompNode3 = FBTBuilder::AddSelector(CompNode);
+			{
+				FBTBuilder::WithServiceLog(CompNode3, 1/*ActivationIndex*/, 2/*DeactivationIndex*/, 3/*TickIndex*/, TEXT("Bool1"));
+
+				FBTBuilder::AddTask(CompNode3, 6, EBTNodeResult::Succeeded);
+			}
+		}
+
+		ExpectedResult.Add(1);
+		ExpectedResult.Add(3);
+		ExpectedResult.Add(6);
+		ExpectedResult.Add(3); // Extra service tick because aux nodes are always ticked before any pending requests are processed
+		ExpectedResult.Add(2);
+		ExpectedResult.Add(4);
+	}
+};
+IMPLEMENT_AI_LATENT_TEST(FAITest_BTSwitchingHigherPrioDuringServiceTick, "System.AI.Behavior Trees.Switch: higher priority during service tick")
+
+struct FAITest_BTSwitchingHigherPrioDuringServiceBecomeRelevant : public FAITest_SimpleBT
+{
+	FAITest_BTSwitchingHigherPrioDuringServiceBecomeRelevant()
+	{
+		UBTCompositeNode& CompNode = FBTBuilder::AddSelector(*BTAsset);
+		{
+			UBTCompositeNode& CompNode2 = FBTBuilder::AddSequence(CompNode);
+			{
+				FBTBuilder::WithDecoratorBlackboard(CompNode, EBasicKeyOperation::Set, EBTFlowAbortMode::LowerPriority, TEXT("Bool1"));
+
+				FBTBuilder::AddTask(CompNode2, 4, EBTNodeResult::Succeeded);
+			}
+			UBTCompositeNode& CompNode3 = FBTBuilder::AddSelector(CompNode);
+			{
+				FBTBuilder::WithServiceLog(CompNode3, 1/*ActivationIndex*/, 2/*DeactivationIndex*/, 3/*TickIndex*/, NAME_None, false, TEXT("Bool1"));
+
+				FBTBuilder::AddTask(CompNode3, 6, EBTNodeResult::Succeeded);
+			}
+		}
+
+		ExpectedResult.Add(1);
+		ExpectedResult.Add(3);
+		ExpectedResult.Add(6);
+		ExpectedResult.Add(3); // Extra service tick because aux nodes are always ticked before any pending requests are processed
+		ExpectedResult.Add(2);
+		ExpectedResult.Add(4);
+	}
+};
+IMPLEMENT_AI_LATENT_TEST(FAITest_BTSwitchingHigherPrioDuringServiceBecomeRelevant, "System.AI.Behavior Trees.Switch: higher priority during service becomes relevant")
+
+struct FAITest_BTSwitchingHigherPrioDuringServiceBecomeRelevantWithTickingTask : public FAITest_SimpleBT
+{
+	FAITest_BTSwitchingHigherPrioDuringServiceBecomeRelevantWithTickingTask()
+	{
+		UBTCompositeNode& CompNode = FBTBuilder::AddSelector(*BTAsset);
+		{
+			UBTCompositeNode& CompNode2 = FBTBuilder::AddSequence(CompNode);
+			{
+				FBTBuilder::WithDecoratorBlackboard(CompNode, EBasicKeyOperation::Set, EBTFlowAbortMode::LowerPriority, TEXT("Bool1"));
+
+				FBTBuilder::AddTask(CompNode2, 4, EBTNodeResult::Succeeded);
+			}
+			UBTCompositeNode& CompNode3 = FBTBuilder::AddSelector(CompNode);
+			{
+				FBTBuilder::WithServiceLog(CompNode3, 1/*ActivationIndex*/, 2/*DeactivationIndex*/, 3/*TickIndex*/, NAME_None, false, TEXT("Bool1"));
+
+				FBTBuilder::AddTask(CompNode3, 6, EBTNodeResult::Succeeded, 10/*ExecutionTicks*/, 222/*LogTickIndex*/);
+			}
+		}
+
+		ExpectedResult.Add(1);
+		ExpectedResult.Add(3);
+		ExpectedResult.Add(6);
+		ExpectedResult.Add(222);
+		ExpectedResult.Add(3); // Extra service tick because aux nodes are always ticked before any pending requests are processed
+		ExpectedResult.Add(2);
+		ExpectedResult.Add(4);
+	}
+};
+IMPLEMENT_AI_LATENT_TEST(FAITest_BTSwitchingHigherPrioDuringServiceBecomeRelevantWithTickingTask, "System.AI.Behavior Trees.Switch: higher priority during service becomes relevant with ticking task")
+
+struct FAITest_BTSwitchingHigherPrioDuringServiceCeaseRelevant : public FAITest_SimpleBT
+{
+	FAITest_BTSwitchingHigherPrioDuringServiceCeaseRelevant()
+	{
+		UBTCompositeNode& CompNode = FBTBuilder::AddSelector(*BTAsset);
+		{
+			UBTCompositeNode& CompNode2 = FBTBuilder::AddSequence(CompNode);
+			{
+				FBTBuilder::WithDecoratorBlackboard(CompNode, EBasicKeyOperation::Set, EBTFlowAbortMode::LowerPriority, TEXT("Bool1"));
+
+				FBTBuilder::AddTask(CompNode2, 4, EBTNodeResult::Succeeded);
+			}
+			UBTCompositeNode& CompNode3 = FBTBuilder::AddSequence(CompNode);
+			{
+				UBTCompositeNode& CompNode4 = FBTBuilder::AddSelector(CompNode3);
+				{
+					FBTBuilder::WithServiceLog(CompNode4, 1/*ActivationIndex*/, 2/*DeactivationIndex*/, 3/*TickIndex*/, NAME_None, false, NAME_None, TEXT("Bool1"));
+
+					FBTBuilder::AddTask(CompNode4, 6, EBTNodeResult::Succeeded);
+				}
+				FBTBuilder::AddTask(CompNode3, 7, EBTNodeResult::Succeeded);
+			}
+		}
+
+		ExpectedResult.Add(1);
+		ExpectedResult.Add(3);
+		ExpectedResult.Add(6);
+		ExpectedResult.Add(3); // Extra service tick because aux nodes are always ticked before any pending requests are processed
+		ExpectedResult.Add(2);
+		ExpectedResult.Add(7);
+		ExpectedResult.Add(4);
+	}
+};
+IMPLEMENT_AI_LATENT_TEST(FAITest_BTSwitchingHigherPrioDuringServiceCeaseRelevant, "System.AI.Behavior Trees.Switch: higher priority during service ceases relevant")
+
+struct FAITest_BTSwitchingHigherPrioDuringTaskServiceTick : public FAITest_SimpleBT
+{
+	FAITest_BTSwitchingHigherPrioDuringTaskServiceTick()
+	{
+		UBTCompositeNode& CompNode = FBTBuilder::AddSelector(*BTAsset);
+		{
+			UBTCompositeNode& CompNode2 = FBTBuilder::AddSequence(CompNode);
+			{
+				FBTBuilder::WithDecoratorBlackboard(CompNode, EBasicKeyOperation::Set, EBTFlowAbortMode::LowerPriority, TEXT("Bool1"));
+
+				FBTBuilder::AddTask(CompNode2, 4, EBTNodeResult::Succeeded);
+			}
+			FBTBuilder::AddTask(CompNode, 6, EBTNodeResult::Succeeded);
+			FBTBuilder::WithTaskServiceLog(CompNode, 1/*ActivationIndex*/, 2/*DeactivationIndex*/, 3/*TickIndex*/, TEXT("Bool1"));
+		}
+
+		ExpectedResult.Add(1);
+		ExpectedResult.Add(3);
+		ExpectedResult.Add(6);
+		ExpectedResult.Add(3); // Extra service tick because aux nodes are always ticked before any pending requests are processed
+		ExpectedResult.Add(2);
+		ExpectedResult.Add(4);
+	}
+};
+IMPLEMENT_AI_LATENT_TEST(FAITest_BTSwitchingHigherPrioDuringTaskServiceTick, "System.AI.Behavior Trees.Switch: higher priority during task service tick")
+
+
+struct FAITest_BTSwitchingHigherPrioDuringTaskServiceBecomeRelevant : public FAITest_SimpleBT
+{
+	FAITest_BTSwitchingHigherPrioDuringTaskServiceBecomeRelevant()
+	{
+		UBTCompositeNode& CompNode = FBTBuilder::AddSelector(*BTAsset);
+		{
+			UBTCompositeNode& CompNode2 = FBTBuilder::AddSequence(CompNode);
+			{
+				FBTBuilder::WithDecoratorBlackboard(CompNode, EBasicKeyOperation::Set, EBTFlowAbortMode::LowerPriority, TEXT("Bool1"));
+
+				FBTBuilder::AddTask(CompNode2, 4, EBTNodeResult::Succeeded);
+			}
+			FBTBuilder::AddTask(CompNode, 6, EBTNodeResult::Succeeded);
+			FBTBuilder::WithTaskServiceLog(CompNode, 1/*ActivationIndex*/, 2/*DeactivationIndex*/, 3/*TickIndex*/, NAME_None, false, TEXT("Bool1"));
+		}
+
+		ExpectedResult.Add(1);
+		ExpectedResult.Add(3);
+		ExpectedResult.Add(6);
+		ExpectedResult.Add(3); // Extra service tick because aux nodes are always ticked before any pending requests are processed
+		ExpectedResult.Add(2);
+		ExpectedResult.Add(4);
+	}
+};
+IMPLEMENT_AI_LATENT_TEST(FAITest_BTSwitchingHigherPrioDuringTaskServiceBecomeRelevant, "System.AI.Behavior Trees.Switch: higher priority during task service becomes relevant")
+
+struct FAITest_BTSwitchingHigherPrioDuringTaskServiceBecomeRelevantWithTickingTask : public FAITest_SimpleBT
+{
+	FAITest_BTSwitchingHigherPrioDuringTaskServiceBecomeRelevantWithTickingTask()
+	{
+		UBTCompositeNode& CompNode = FBTBuilder::AddSelector(*BTAsset);
+		{
+			UBTCompositeNode& CompNode2 = FBTBuilder::AddSequence(CompNode);
+			{
+				FBTBuilder::WithDecoratorBlackboard(CompNode, EBasicKeyOperation::Set, EBTFlowAbortMode::LowerPriority, TEXT("Bool1"));
+
+				FBTBuilder::AddTask(CompNode2, 4, EBTNodeResult::Succeeded);
+			}
+			FBTBuilder::AddTask(CompNode, 6, EBTNodeResult::Succeeded, 10/*ExecutionTicks*/, 222/*LogTickIndex*/);
+			FBTBuilder::WithTaskServiceLog(CompNode, 1/*ActivationIndex*/, 2/*DeactivationIndex*/, 3/*TickIndex*/, NAME_None, false, TEXT("Bool1"));
+		}
+
+		ExpectedResult.Add(1);
+		ExpectedResult.Add(3);
+		ExpectedResult.Add(6);
+		ExpectedResult.Add(222);
+		ExpectedResult.Add(3); // Extra service tick because aux nodes are always ticked before any pending requests are processed
+		ExpectedResult.Add(2);
+		ExpectedResult.Add(4);
+	}
+};
+IMPLEMENT_AI_LATENT_TEST(FAITest_BTSwitchingHigherPrioDuringTaskServiceBecomeRelevantWithTickingTask, "System.AI.Behavior Trees.Switch: higher priority during task service becomes relevant with ticking task")
+
+struct FAITest_BTSwitchingHigherPrioDuringTaskServiceCeaseRelevant : public FAITest_SimpleBT
+{
+	FAITest_BTSwitchingHigherPrioDuringTaskServiceCeaseRelevant()
+	{
+		UBTCompositeNode& CompNode = FBTBuilder::AddSelector(*BTAsset);
+		{
+			UBTCompositeNode& CompNode2 = FBTBuilder::AddSequence(CompNode);
+			{
+				FBTBuilder::WithDecoratorBlackboard(CompNode, EBasicKeyOperation::Set, EBTFlowAbortMode::LowerPriority, TEXT("Bool1"));
+
+				FBTBuilder::AddTask(CompNode2, 4, EBTNodeResult::Succeeded);
+			}
+			UBTCompositeNode& CompNode3 = FBTBuilder::AddSequence(CompNode);
+			{
+				FBTBuilder::AddTask(CompNode3, 6, EBTNodeResult::Succeeded);
+				FBTBuilder::WithTaskServiceLog(CompNode3, 1/*ActivationIndex*/, 2/*DeactivationIndex*/, 3/*TickIndex*/, NAME_None, false, NAME_None, TEXT("Bool1"));
+				FBTBuilder::AddTask(CompNode3, 7, EBTNodeResult::Succeeded);
+			}
+		}
+
+		ExpectedResult.Add(1);
+		ExpectedResult.Add(3);
+		ExpectedResult.Add(6);
+		ExpectedResult.Add(3); // Extra service tick because aux nodes are always ticked before any pending requests are processed
+		ExpectedResult.Add(2);
+		ExpectedResult.Add(7);
+		ExpectedResult.Add(4);
+	}
+};
+IMPLEMENT_AI_LATENT_TEST(FAITest_BTSwitchingHigherPrioDuringTaskServiceCeaseRelevant, "System.AI.Behavior Trees.Switch: higher priority during task service ceases relevant")
+
 struct FAITest_BTAbortSelfFail : public FAITest_SimpleBT
 {
 	FAITest_BTAbortSelfFail()
diff --git a/Engine/Source/Runtime/AIModule/Classes/BehaviorTree/BehaviorTreeComponent.h b/Engine/Source/Runtime/AIModule/Classes/BehaviorTree/BehaviorTreeComponent.h
--- a/Engine/Source/Runtime/AIModule/Classes/BehaviorTree/BehaviorTreeComponent.h
+++ b/Engine/Source/Runtime/AIModule/Classes/BehaviorTree/BehaviorTreeComponent.h
@@ -68,12 +68,6 @@ struct FBTPendingExecutionInfo
 	void Unlock() { bLocked = false; }
 };
 
-struct FBTPendingAuxNodesUnregisterInfo
-{
-	/** list of node index ranges pending aux nodes unregistration */
-	TArray<FBTNodeIndexRange> Ranges;
-};
-
 struct FBTTreeStartInfo
 {
 	UBehaviorTree* Asset;
@@ -139,7 +133,11 @@ public:
 	void RequestExecution(EBTNodeResult::Type ContinueWithResult);
 
 	/** request unregistration of aux nodes in the specified branch */
+	UE_DEPRECATED(5.0, "This function is deprecated. Please use RequestBranchDeactivation instead.")
 	void RequestUnregisterAuxNodesInBranch(const UBTCompositeNode* Node);
+		
+	/** request branch deactivation: helper for decorator */
+	void RequestBranchDeactivation(const UBTDecorator& RequestedBy);
 
 	/** finish latent execution or abort */
 	void OnTaskFinished(const UBTTaskNode* TaskNode, EBTNodeResult::Type TaskResult);
@@ -165,7 +163,6 @@ public:
 	void UnregisterAuxNodesInRange(const FBTNodeIndex& FromIndex, const FBTNodeIndex& ToIndex);
 
 	/** unregister all aux nodes in branch of tree */
-	UE_DEPRECATED(4.26, "This function is deprecated. Please use RequestUnregisterAuxNodesInBranch instead.")
 	void UnregisterAuxNodesInBranch(const UBTCompositeNode* Node, bool bApplyImmediately = true);
 
 	/** BEGIN UActorComponent overrides */
@@ -271,8 +268,8 @@ protected:
 	/** result of ExecutionRequest, will be applied when current task finish aborting */
 	FBTPendingExecutionInfo PendingExecution;
 
-	/** list of all pending aux nodes unregistration requests */
-	FBTPendingAuxNodesUnregisterInfo PendingUnregisterAuxNodesRequests;
+	/** list of all pending branch deactivation requests */
+	TArray<const UBTNode*> PendingBranchesToDeactivate;
 
 	/** stored data for starting new tree, waits until previously running finishes aborting */
 	FBTTreeStartInfo TreeStartInfo;
@@ -331,6 +328,9 @@ protected:
 	/** if set, execution requests will be postponed */
 	uint8 bIsPaused : 1;
 
+	/** if set, all branch deactivation requests will be queued up */
+	uint8 bBranchDeactivationSuspended : 1;
+
 	/** push behavior tree instance on execution stack
 	 *	@NOTE: should never be called out-side of BT execution, meaning only BT tasks can push another BT instance! */
 	bool PushInstance(UBehaviorTree& TreeAsset);
@@ -380,12 +380,6 @@ protected:
 	/** apply pending tree initialization */
 	void ProcessPendingInitialize();
 
-	/**
-	 * apply pending unregister aux nodes requests
-	 * @return true if some request were processed, false otherwise
-	 */
-	bool ProcessPendingUnregister();
-
 	/** restore state of tree to state before search */
 	void RollbackSearchChanges();
 
@@ -418,6 +412,29 @@ protected:
 	/** Return NodeA's relative priority in regards to NodeB */
 	EBTNodeRelativePriority CalculateRelativePriority(const UBTNode* NodeA, const UBTNode* NodeB) const;
 
+	/** Deactivate a branch as the decorator condition is not passing anymore */
+	void DeactivateBranch(const UBTDecorator& RequestedBy);
+
+	/** Suspend any branch deactivation and queue them to be processed later by ResumeBranchDeactivation() */
+	void SuspendBranchDeactivation();
+
+	/** Resume branch deactivation and execute all the queued up ones */
+	void ResumeBranchDeactivation();
+
+	struct FBTSuspendBranchDeactivationScoped
+	{
+		FBTSuspendBranchDeactivationScoped(UBehaviorTreeComponent& InBTComp)
+			: BTComp(InBTComp)
+		{
+			BTComp.SuspendBranchDeactivation();
+		}
+		~FBTSuspendBranchDeactivationScoped()
+		{
+			BTComp.ResumeBranchDeactivation();
+		}
+		UBehaviorTreeComponent& BTComp;
+	};
+
 	friend UBTNode;
 	friend UBTCompositeNode;
 	friend UBTTaskNode;
diff --git a/Engine/Source/Runtime/AIModule/Private/BehaviorTree/BTDecorator.cpp b/Engine/Source/Runtime/AIModule/Private/BehaviorTree/BTDecorator.cpp
--- a/Engine/Source/Runtime/AIModule/Private/BehaviorTree/BTDecorator.cpp
+++ b/Engine/Source/Runtime/AIModule/Private/BehaviorTree/BTDecorator.cpp
@@ -102,29 +102,24 @@ void UBTDecorator::ConditionalFlowAbort(UBehaviorTreeComponent& OwnerComp, EBTDe
 	const bool bAbortPending = OwnerComp.IsAbortPending();
 	const bool bAlwaysRequestWhenPassing = (RequestMode == EBTDecoratorAbortRequest::ConditionPassing);
 
-	const bool bLogRestart = (bIsExecutingBranch != bPass) || (bIsExecutingBranch && bPass && (bAlwaysRequestWhenPassing || bAbortPending));
+	const bool bLogRequestExecution = (!bIsExecutingBranch && bPass) || (bIsExecutingBranch && bPass && (bAlwaysRequestWhenPassing || bAbortPending));
 	UE_VLOG(OwnerComp.GetOwner(), LogBehaviorTree, Verbose, TEXT("%s, ConditionalFlowAbort(%s) pass:%d executingBranch:%d abortPending:%d => %s"),
 		*UBehaviorTreeTypes::DescribeNodeHelper(this),
 		bAlwaysRequestWhenPassing ? TEXT("always when passing") : TEXT("on change"),
 		bPass ? 1 : 0,
 		bIsExecutingBranch ? 1 : 0,
 		bAbortPending ? 1 : 0,
-		bLogRestart ? TEXT("restart") : TEXT("skip"));
+		!bPass ? TEXT("request branch deactivation") : (bLogRequestExecution ? TEXT("request execution") : TEXT("skip")));
 
-	if (bIsExecutingBranch != bPass)
+	if (!bPass)
 	{
-		OwnerComp.RequestExecution(this);
+		OwnerComp.RequestBranchDeactivation(*this);
 	}
-	else if (!bIsExecutingBranch && !bPass && GetParentNode() && GetParentNode()->Children.IsValidIndex(GetChildIndex()))
+	else if (!bIsExecutingBranch)
 	{
-		// this condition here is to remove all active observers _BELOW_ this node
-		// because if this condition failed we no longer want to react to child-conditions
-		// value changes anyway since their nodes execution will be blocked by this condition 
-		// during tree search
-		const UBTCompositeNode* BranchRoot = GetParentNode()->Children[GetChildIndex()].ChildComposite;
-		OwnerComp.RequestUnregisterAuxNodesInBranch(BranchRoot);
+		OwnerComp.RequestExecution(this);
 	}
-	else if (bIsExecutingBranch && bPass && (bAlwaysRequestWhenPassing || bAbortPending))
+	else if (bAlwaysRequestWhenPassing || bAbortPending)
 	{
 		// force result Aborted to restart from this decorator
 		OwnerComp.RequestExecution(GetParentNode(), InstanceIdx, this, GetChildIndex(), EBTNodeResult::Aborted);
diff --git a/Engine/Source/Runtime/AIModule/Private/BehaviorTree/BehaviorTreeComponent.cpp b/Engine/Source/Runtime/AIModule/Private/BehaviorTree/BehaviorTreeComponent.cpp
--- a/Engine/Source/Runtime/AIModule/Private/BehaviorTree/BehaviorTreeComponent.cpp
+++ b/Engine/Source/Runtime/AIModule/Private/BehaviorTree/BehaviorTreeComponent.cpp
@@ -71,6 +71,7 @@ UBehaviorTreeComponent::UBehaviorTreeComponent(const FObjectInitializer& ObjectI
 	bWantsInitializeComponent = true; 
 	bIsRunning = false;
 	bIsPaused = false;
+	bBranchDeactivationSuspended = false;
 
 	// Adding hook for bespoke framepro BT timings for BR
 #if !UE_BUILD_SHIPPING
@@ -303,11 +304,14 @@ void UBehaviorTreeComponent::StopTree(EBTStopMode::Type StopMode)
 			FBehaviorTreeInstance& InstanceInfo = InstanceStack[InstanceIndex];
 
 			// notify active aux nodes
-			InstanceInfo.ExecuteOnEachAuxNode([&InstanceInfo, this](const UBTAuxiliaryNode& AuxNode)
+			{
+				FBTSuspendBranchDeactivationScoped ScopedSuspend(*this);
+				InstanceInfo.ExecuteOnEachAuxNode([&InstanceInfo, this](const UBTAuxiliaryNode& AuxNode)
 				{
 					uint8* NodeMemory = AuxNode.GetNodeMemory<uint8>(InstanceInfo);
 					AuxNode.WrappedOnCeaseRelevant(*this, NodeMemory);
 				});
+			}
 			InstanceInfo.ResetActiveAuxNodes();
 
 			// notify active parallel tasks
@@ -696,14 +700,88 @@ EBTTaskStatus::Type UBehaviorTreeComponent::GetTaskStatus(const UBTTaskNode* Tas
 
 void UBehaviorTreeComponent::RequestUnregisterAuxNodesInBranch(const UBTCompositeNode* Node)
 {
-	const int32 InstanceIdx = FindInstanceContainingNode(Node);
-	if (InstanceIdx != INDEX_NONE)
+	if (!Node)
+	{
+		return;
+	}
+
+	if (bBranchDeactivationSuspended)
 	{
-		PendingUnregisterAuxNodesRequests.Ranges.Emplace(
-			FBTNodeIndex(InstanceIdx, Node->GetExecutionIndex()),
-			FBTNodeIndex(InstanceIdx, Node->GetLastExecutionIndex()));
+		PendingBranchesToDeactivate.Add(Node);
+		return;
+	}
 
-		ScheduleNextTick(0.0f);
+	UnregisterAuxNodesInBranch(Node, true/*bApplyImmediately*/);
+}
+
+void UBehaviorTreeComponent::DeactivateBranch(const UBTDecorator& RequestedBy)
+{
+	if (IsExecutingBranch(&RequestedBy, RequestedBy.GetChildIndex()))
+	{
+		UE_VLOG(GetOwner(), LogBehaviorTree, Verbose, TEXT("%s, Branch deactivation resulted in new execution request"), *UBehaviorTreeTypes::DescribeNodeHelper(&RequestedBy));
+		RequestExecution(&RequestedBy);
+	}
+	else if (ensureMsgf(RequestedBy.GetParentNode() && RequestedBy.GetParentNode()->Children.IsValidIndex(RequestedBy.GetChildIndex()), 
+				TEXT("The decorator %s does not have a parent or is not a valid child."), *UBehaviorTreeTypes::DescribeNodeHelper(&RequestedBy)))
+	{
+		UE_VLOG(GetOwner(), LogBehaviorTree, Verbose, TEXT("%s, Branch deactivation resulted in aux nodes unregistration"), *UBehaviorTreeTypes::DescribeNodeHelper(&RequestedBy));
+		if (const UBTCompositeNode* BranchRoot = RequestedBy.GetParentNode()->Children[RequestedBy.GetChildIndex()].ChildComposite)
+		{
+			UnregisterAuxNodesInBranch(BranchRoot, true/*bApplyImmediately*/);
+		}
+	}
+	else
+	{
+		UE_VLOG(GetOwner(), LogBehaviorTree, Error, TEXT("The decorator %s does not have a parent or is not a valid child."), *UBehaviorTreeTypes::DescribeNodeHelper(&RequestedBy));
+	}
+}
+
+void UBehaviorTreeComponent::RequestBranchDeactivation(const UBTDecorator& RequestedBy)
+{
+	if (bBranchDeactivationSuspended)
+	{
+		UE_VLOG(GetOwner(), LogBehaviorTree, Verbose, TEXT("%s, Branch deactivation queued up"), *UBehaviorTreeTypes::DescribeNodeHelper(&RequestedBy));
+		PendingBranchesToDeactivate.Add(&RequestedBy);
+		return;
+	}
+
+	DeactivateBranch(RequestedBy);
+}
+
+void UBehaviorTreeComponent::SuspendBranchDeactivation()
+{
+	UE_VLOG(GetOwner(), LogBehaviorTree, VeryVerbose, TEXT("Suspending branch deactivation."));
+	checkf(!bBranchDeactivationSuspended, TEXT("This logic does not support re-entrance"));
+	bBranchDeactivationSuspended = true;
+}
+
+void UBehaviorTreeComponent::ResumeBranchDeactivation()
+{
+	UE_VLOG(GetOwner(), LogBehaviorTree, VeryVerbose, TEXT("Resuming branch deactivation."));
+	checkf(bBranchDeactivationSuspended, TEXT("Expecting SuspendBranchDeactivation() be called before calling resume"));
+	bBranchDeactivationSuspended = false;
+
+	// Flushing any pending branch deactivations
+	while (PendingBranchesToDeactivate.Num() > 0)
+	{
+		TArray<const UBTNode*> PendingBranchesToDeactivateToProcess(MoveTemp(PendingBranchesToDeactivate));
+		PendingBranchesToDeactivate.Reset();
+		for (const UBTNode* Node : PendingBranchesToDeactivateToProcess)
+		{
+			if (const UBTDecorator* RequestedDecorator = Cast<UBTDecorator>(Node))
+			{
+				DeactivateBranch(*RequestedDecorator);
+			}
+			//////////////////////////////////////////////////////////////////////
+			// An attempt at backward compatibility 
+			// If anybody is still using deprecated method RequestUnregisterAuxNodesInBranch
+			else if (const UBTCompositeNode* BranchRoot = Cast<UBTCompositeNode>(Node))
+			{
+				UnregisterAuxNodesInBranch(BranchRoot, true/*bApplyImmediately*/);
+			}
+			//
+			//////////////////////////////////////////////////////////////////////
+		}
 	}
 }
 
@@ -905,15 +983,49 @@ void UBehaviorTreeComponent::RequestExecution(UBTCompositeNode* RequestedOn, int
 	ExecutionIdx.ExecutionIndex = RequestedBy->GetExecutionIndex();
 	uint16 LastExecutionIndex = MAX_uint16;
 
-	// make sure that the request is not coming from a node that has pending unregistration since it won't be accessible anymore
-	for (const FBTNodeIndexRange& Range : PendingUnregisterAuxNodesRequests.Ranges)
+	// make sure that the request is not coming from a node that has pending branch deactivation since it won't be accessible anymore
+	if (bBranchDeactivationSuspended)
 	{
-		if (Range.Contains(ExecutionIdx))
+		for (const UBTNode* Node : PendingBranchesToDeactivate)
 		{
-			UE_VLOG(GetOwner(), LogBehaviorTree, Log, TEXT("> skip: request by %s that is in pending unregister aux nodes range %s"), *ExecutionIdx.Describe(), *Range.Describe());
-			return;
+			const UBTCompositeNode* BranchRoot = nullptr;
+			if (const UBTDecorator* Decorator = Cast<UBTDecorator>(Node))
+			{
+				// This check was already in previous version, it will ensure and output a vlog later in the DeactivateBranch, no need to do anything now.
+				if (Decorator->GetParentNode() && Decorator->GetParentNode()->Children.IsValidIndex(Decorator->GetChildIndex()))
+				{
+					BranchRoot = Decorator->GetParentNode()->Children[Decorator->GetChildIndex()].ChildComposite;
+				}
+			}
+			//////////////////////////////////////////////////////////////////////
+			// An attempt at backward compatibility 
+			// If anybody is still using deprecated method RequestUnregisterAuxNodesInBranch
+			else if (const UBTCompositeNode* CompNode = Cast<UBTCompositeNode>(Node))
+			{
+				BranchRoot = CompNode;
+			}
+			//
+			//////////////////////////////////////////////////////////////////////
+
+			if (BranchRoot)
+			{
+				const int32 BranchRootInstanceIdx = FindInstanceContainingNode(BranchRoot);
+				if (BranchRootInstanceIdx != INDEX_NONE)
+				{
+					FBTNodeIndexRange Range( FBTNodeIndex(BranchRootInstanceIdx, BranchRoot->GetExecutionIndex()), FBTNodeIndex(BranchRootInstanceIdx, BranchRoot->GetLastExecutionIndex()));
+					if (Range.Contains(ExecutionIdx))
+					{
+						UE_VLOG(GetOwner(), LogBehaviorTree, Log, TEXT("> skip: request by %s(%s) is in deactivated branch %s(%s) and was deactivated by %s"), *UBehaviorTreeTypes::DescribeNodeHelper(RequestedBy), *ExecutionIdx.Describe(), *UBehaviorTreeTypes::DescribeNodeHelper(BranchRoot), *Range.Describe(), *UBehaviorTreeTypes::DescribeNodeHelper(Node));
+						return;
+					}
+				}
+			}
 		}
 	}
+	else
+	{
+		checkf(PendingBranchesToDeactivate.Num() == 0, TEXT("All pending branches should have been flushed before requesting an execution"));
+	}
 
 	if (bSwitchToHigherPriority && RequestedByChildIndex >= 0)
 	{
@@ -1345,21 +1457,23 @@ void UBehaviorTreeComponent::TickComponent(float DeltaTime, enum ELevelTick Tick
 	check(this != nullptr && this->IsPendingKill() == false);
 	float NextNeededDeltaTime = FLT_MAX;
 
-	// process all auxiliary nodes unregister requests
-	bDoneSomething |= ProcessPendingUnregister();
+	checkf(PendingBranchesToDeactivate.Num() == 0, TEXT("Pending branches should always be flushed immediately with the new system"))
 
 	// tick active auxiliary nodes (in execution order, before task)
 	// do it before processing execution request to give BP driven logic chance to accumulate execution requests
 	// newly added aux nodes are ticked as part of SearchData application
-	for (int32 InstanceIndex = 0; InstanceIndex < InstanceStack.Num(); InstanceIndex++)
 	{
-		FBehaviorTreeInstance& InstanceInfo = InstanceStack[InstanceIndex];
-		InstanceInfo.ExecuteOnEachAuxNode([&InstanceInfo, this, &bDoneSomething, DeltaTime, &NextNeededDeltaTime](const UBTAuxiliaryNode& AuxNode)
-			{
-				uint8* NodeMemory = AuxNode.GetNodeMemory<uint8>(InstanceInfo);
-				SCOPE_CYCLE_UOBJECT(AuxNode, &AuxNode);
-				bDoneSomething |= AuxNode.WrappedTickNode(*this, NodeMemory, DeltaTime, NextNeededDeltaTime);
-			});
+		FBTSuspendBranchDeactivationScoped ScopedSuspend(*this);
+		for (int32 InstanceIndex = 0; InstanceIndex < InstanceStack.Num(); InstanceIndex++)
+		{
+			FBehaviorTreeInstance& InstanceInfo = InstanceStack[InstanceIndex];
+			InstanceInfo.ExecuteOnEachAuxNode([&InstanceInfo, this, &bDoneSomething, DeltaTime, &NextNeededDeltaTime](const UBTAuxiliaryNode& AuxNode)
+				{
+					uint8* NodeMemory = AuxNode.GetNodeMemory<uint8>(InstanceInfo);
+					SCOPE_CYCLE_UOBJECT(AuxNode, &AuxNode);
+					bDoneSomething |= AuxNode.WrappedTickNode(*this, NodeMemory, DeltaTime, NextNeededDeltaTime);
+				});
+		}
 	}
 
 	bool bActiveAuxiliaryNodeDTDirty = false;
@@ -1802,6 +1916,9 @@ void UBehaviorTreeComponent::ProcessPendingExecution()
 	const FBTNodeIndex NextTaskIdx = SavedInfo.NextTask ? FBTNodeIndex(ActiveInstanceIdx, SavedInfo.NextTask->GetExecutionIndex()) : FBTNodeIndex(0, 0);
 	UnregisterAuxNodesUpTo(NextTaskIdx);
 
+	// Suspending any branch deactivation as it is impossible for decorators to have the right answer if they are in an executing branch or not.
+	SuspendBranchDeactivation();
+
 	// change aux nodes
 	ApplySearchData(SavedInfo.NextTask);
 
@@ -1820,10 +1937,12 @@ void UBehaviorTreeComponent::ProcessPendingExecution()
 	// validate active instance as well, execution can be delayed AND can have AbortCurrentTask call before using instance index
 	if (SavedInfo.NextTask && InstanceStack.IsValidIndex(ActiveInstanceIdx))
 	{
+		// ResumeBranchDeactivation() is done inside ExecuteTask after the active task is set but before we execute the task.
 		ExecuteTask(SavedInfo.NextTask);
 	}
 	else
 	{
+		ResumeBranchDeactivation();
 		OnTreeFinished();
 	}
 }
@@ -1987,34 +2106,13 @@ void UBehaviorTreeComponent::UnregisterAuxNodesInBranch(const UBTCompositeNode*
 
 		if (bApplyImmediately)
 		{
+			FBTSuspendBranchDeactivationScoped ScopedSuspend(*this);
 			ApplySearchUpdates(SearchData.PendingUpdates, 0);
 			SearchData.PendingUpdates = UpdateListCopy;
 		}
 	}
 }
 
-bool UBehaviorTreeComponent::ProcessPendingUnregister()
-{
-	if (PendingUnregisterAuxNodesRequests.Ranges.Num() == 0)
-	{
-		// no work done
-		return false;
-	}
-
-	TGuardValue<TArray<FBehaviorTreeSearchUpdate>> ScopedList(SearchData.PendingUpdates, {});
-
-	for (const FBTNodeIndexRange& Range : PendingUnregisterAuxNodesRequests.Ranges)
-	{
-		UnregisterAuxNodesInRange(Range.FromIndex, Range.ToIndex);
-	}
-	PendingUnregisterAuxNodesRequests = {};
-
-	ApplySearchUpdates(SearchData.PendingUpdates, 0);
-
-	// has done work
-	return true;
-}
-
 void UBehaviorTreeComponent::ExecuteTask(UBTTaskNode* TaskNode)
 {
 	SCOPE_CYCLE_COUNTER(STAT_AI_BehaviorTree_ExecutionTime);
@@ -2055,6 +2153,10 @@ void UBehaviorTreeComponent::ExecuteTask(UBTTaskNode* TaskNode)
 	ActiveInstance.ActiveNode = TaskNode;
 	ActiveInstance.ActiveNodeType = EBTActiveNode::ActiveTask;
 
+	// Is is now ok to resume any branch deactivation as the new active instance is set.
+	// Before that, decorators evaluating the IsExecutingBranch would be wrong.
+	ResumeBranchDeactivation();
+
 	// make a snapshot for debugger
 	StoreDebuggerExecutionStep(EBTExecutionSnap::Regular);
 
