Fix a new BT bug found which was not trackable before, but now with latest code it is now being catch by an ensure.

diff --git a/Engine/Source/Developer/AITestSuite/Private/Tests/BTTest.cpp b/Engine/Source/Developer/AITestSuite/Private/Tests/BTTest.cpp
--- a/Engine/Source/Developer/AITestSuite/Private/Tests/BTTest.cpp
+++ b/Engine/Source/Developer/AITestSuite/Private/Tests/BTTest.cpp
@@ -689,6 +689,39 @@ struct FAITest_BTStartBranchDuringServiceCeaseRelevant : public FAITest_SimpleBT
 };
 IMPLEMENT_AI_LATENT_TEST(FAITest_BTStartBranchDuringServiceCeaseRelevant, "System.AI.Behavior Trees.Toggle: start new branch during service cease relevant")
 
+struct FAITest_BTStartBranchDuringServiceCeaseRelevant2 : public FAITest_SimpleBT
+{
+	FAITest_BTStartBranchDuringServiceCeaseRelevant2()
+	{
+		UBTCompositeNode& CompNode = FBTBuilder::AddSelector(*BTAsset);
+		{
+			UBTCompositeNode& CompNode2 = FBTBuilder::AddSelector(CompNode);
+			{
+				FBTBuilder::AddTask(CompNode2, 6/*LogIndex*/, EBTNodeResult::Succeeded, 2/* ExecutionTicks*/);
+				FBTBuilder::WithDecoratorBlackboard(CompNode2, EBasicKeyOperation::Set, EBTFlowAbortMode::Both, TEXT("Bool1"));
+
+				UBTCompositeNode& CompNode3 = FBTBuilder::AddSelector(CompNode2);
+				{
+					FBTBuilder::AddTask(CompNode3, 5/*LogIndex*/, EBTNodeResult::Succeeded);
+					FBTBuilder::WithDecoratorBlackboard(CompNode3, EBasicKeyOperation::Set, EBTFlowAbortMode::Both, TEXT("Bool2"));
+
+					FBTBuilder::AddTaskToggleFlag(CompNode3, EBTNodeResult::InProgress, TEXT("Bool1"), 1/*NumOfToggles*/);
+					FBTBuilder::WithTaskServiceLog(CompNode3, 1 /*ActivationIndex*/, 2 /*DeactivationIndex*/, 3 /*TickIndex*/, NAME_None /*TickBoolKeyName*/, false /*bCallTickOnSearchStart*/,
+						NAME_None /*BecomeRelevantBoolKeyeyName*/, TEXT("Bool2") /*CeaseRelevantBoolKeyName*/);
+				}
+			}
+			FBTBuilder::AddTask(CompNode, 7/*LogIndex*/, EBTNodeResult::Succeeded);
+		}
+			
+		ExpectedResult.Add(1);
+		ExpectedResult.Add(3);
+		ExpectedResult.Add(3);
+		ExpectedResult.Add(2);
+		ExpectedResult.Add(6);
+	}
+};
+IMPLEMENT_AI_LATENT_TEST(FAITest_BTStartBranchDuringServiceCeaseRelevant2, "System.AI.Behavior Trees.Toggle: start new branch during service cease relevant (low prio)")
+
 struct FAITest_BTAbortingDuringServiceBecomeRelevant : public FAITest_SimpleBT
 {
 	FAITest_BTAbortingDuringServiceBecomeRelevant()
diff --git a/Engine/Source/Runtime/AIModule/Private/BehaviorTree/BehaviorTreeComponent.cpp b/Engine/Source/Runtime/AIModule/Private/BehaviorTree/BehaviorTreeComponent.cpp
--- a/Engine/Source/Runtime/AIModule/Private/BehaviorTree/BehaviorTreeComponent.cpp
+++ b/Engine/Source/Runtime/AIModule/Private/BehaviorTree/BehaviorTreeComponent.cpp
@@ -749,6 +749,13 @@ void UBehaviorTreeComponent::ResumeBranchActions()
 				case EBTBranchAction::DecoratorEvaluate:
 				{
 					const UBTDecorator* RequestedBy = CastChecked<UBTDecorator>(Info.Node);
+
+					// Since we have been queued up, decorator might have been removed from active nodes, need to make sure it is still there.
+					if (!IsAuxNodeActive(RequestedBy))
+					{
+						break;
+					}
+
 					EvaluateBranch(*RequestedBy);
 					break;
 				}
@@ -757,12 +764,26 @@ void UBehaviorTreeComponent::ResumeBranchActions()
 				case EBTBranchAction::DecoratorActivate_EvenIfExecuting:
 				{
 					const UBTDecorator* RequestedBy = CastChecked<UBTDecorator>(Info.Node);
+
+					// Since we have been queued up, decorator might have been removed from active nodes, need to make sure it is still there.
+					if (!IsAuxNodeActive(RequestedBy))
+					{
+						break;
+					}
+
 					ActivateBranch(*RequestedBy, Info.Action == EBTBranchAction::DecoratorActivate_EvenIfExecuting /*bForceRequestEvenIfExecuting*/);
 					break;
 				}
 				case EBTBranchAction::DecoratorDeactivate:
 				{
 					const UBTDecorator* RequestedBy = CastChecked<UBTDecorator>(Info.Node);
+
+					// Since we have been queued up, decorator might have been removed from active nodes, need to make sure it is still there.
+					if (!IsAuxNodeActive(RequestedBy))
+					{
+						break;
+					}
+
 					DeactivateBranch(*RequestedBy);
 					break;
 				}
@@ -3295,4 +3316,4 @@ bool UBehaviorTreeComponent::IsDebuggerActive()
 #else
 	return false;
 #endif
-}
+}
\ No newline at end of file
