Better support of stoptree while already being in the behavior tree

diff --git a/Engine/Source/Developer/AITestSuite/Classes/BehaviorTree/TestBTService_StopTree.h b/Engine/Source/Developer/AITestSuite/Classes/BehaviorTree/TestBTService_StopTree.h
new file mode 100644
--- /dev/null
+++ b/Engine/Source/Developer/AITestSuite/Classes/BehaviorTree/TestBTService_StopTree.h
@@ -0,0 +1,32 @@
+// Copyright Epic Games, Inc. All Rights Reserved.
+#pragma once
+
+#include "BehaviorTree/BTService.h"
+#include "TestBTService_StopTree.generated.h"
+
+UENUM()
+namespace EBTTestServiceStopTree
+{
+	enum Type
+	{
+		DuringBecomeRelevant,
+		DuringTick,
+		DuringCeaseRelevant,
+	};
+}
+
+UCLASS(meta = (HiddenNode))
+class UTestBTService_StopTree : public UBTService
+{
+	GENERATED_UCLASS_BODY()
+
+	UPROPERTY()
+	int32 LogIndex;
+
+	UPROPERTY()
+	TEnumAsByte<EBTTestServiceStopTree::Type> StopTimming;
+
+	virtual void OnBecomeRelevant(UBehaviorTreeComponent& OwnerComp, uint8* NodeMemory) override;
+	virtual void OnCeaseRelevant(UBehaviorTreeComponent& OwnerComp, uint8* NodeMemory) override;
+	virtual void TickNode(UBehaviorTreeComponent& OwnerComp, uint8* NodeMemory, float DeltaSeconds) override;
+};
diff --git a/Engine/Source/Developer/AITestSuite/Classes/BehaviorTree/TestBTTask_StopTree.h b/Engine/Source/Developer/AITestSuite/Classes/BehaviorTree/TestBTTask_StopTree.h
new file mode 100644
--- /dev/null
+++ b/Engine/Source/Developer/AITestSuite/Classes/BehaviorTree/TestBTTask_StopTree.h
@@ -0,0 +1,41 @@
+// Copyright Epic Games, Inc. All Rights Reserved.
+#pragma once
+
+#include "BehaviorTree/BTTaskNode.h"
+#include "TestBTTask_StopTree.generated.h"
+
+UENUM()
+namespace EBTTestTaskStopTree
+{
+	enum Type
+	{
+		DuringExecute,
+		DuringTick,
+		DuringAbort,
+		DuringFinish,
+	};
+}
+
+UCLASS(meta=(HiddenNode))
+class UTestBTTask_StopTree : public UBTTaskNode
+{
+	GENERATED_BODY()
+
+public:
+	UTestBTTask_StopTree(const FObjectInitializer& ObjectInitializer);
+
+	virtual EBTNodeResult::Type ExecuteTask(UBehaviorTreeComponent& OwnerComp, uint8* NodeMemory) override;
+	virtual EBTNodeResult::Type AbortTask(UBehaviorTreeComponent& OwnerComp, uint8* NodeMemory) override;
+
+	virtual void TickTask(class UBehaviorTreeComponent& OwnerComp, uint8* NodeMemory, float DeltaSeconds) override;
+	virtual void OnTaskFinished(UBehaviorTreeComponent& OwnerComp, uint8* NodeMemory, EBTNodeResult::Type TaskResult) override;
+
+	UPROPERTY()
+	TEnumAsByte<EBTTestTaskStopTree::Type> StopTimming;
+
+	UPROPERTY()
+	int32 LogIndex;
+
+	UPROPERTY()
+	TEnumAsByte<EBTNodeResult::Type> LogResult;
+};
diff --git a/Engine/Source/Developer/AITestSuite/Private/BehaviorTree/TestBTService_StopTree.cpp b/Engine/Source/Developer/AITestSuite/Private/BehaviorTree/TestBTService_StopTree.cpp
new file mode 100644
--- /dev/null
+++ b/Engine/Source/Developer/AITestSuite/Private/BehaviorTree/TestBTService_StopTree.cpp
@@ -0,0 +1,51 @@
+// Copyright Epic Games, Inc. All Rights Reserved.
+
+#include "BehaviorTree/TestBTService_StopTree.h"
+#include "MockAI_BT.h"
+
+UTestBTService_StopTree::UTestBTService_StopTree(const FObjectInitializer& ObjectInitializer) : Super(ObjectInitializer)
+{
+	NodeName = "StopTreeService";
+
+	INIT_SERVICE_NODE_NOTIFY_FLAGS();
+
+	LogIndex = INDEX_NONE;
+	StopTimming = EBTTestServiceStopTree::DuringTick;
+
+	// Force the service to tick every frame
+	Interval = 0.0f;
+	RandomDeviation = 0.0f;
+}
+
+void UTestBTService_StopTree::OnBecomeRelevant(UBehaviorTreeComponent& OwnerComp, uint8* NodeMemory)
+{
+	Super::OnBecomeRelevant(OwnerComp, NodeMemory);
+
+	if (StopTimming == EBTTestServiceStopTree::DuringBecomeRelevant)
+	{
+		UMockAI_BT::ExecutionLog.Add(LogIndex);
+		OwnerComp.StopTree();
+	}
+}
+
+void UTestBTService_StopTree::OnCeaseRelevant(UBehaviorTreeComponent& OwnerComp, uint8* NodeMemory)
+{
+	Super::OnCeaseRelevant(OwnerComp, NodeMemory);
+
+	if (StopTimming == EBTTestServiceStopTree::DuringCeaseRelevant)
+	{
+		UMockAI_BT::ExecutionLog.Add(LogIndex);
+		OwnerComp.StopTree();
+	}
+}
+
+void UTestBTService_StopTree::TickNode(UBehaviorTreeComponent& OwnerComp, uint8* NodeMemory, float DeltaSeconds)
+{
+	Super::TickNode(OwnerComp, NodeMemory, DeltaSeconds);
+
+	if (StopTimming == EBTTestServiceStopTree::DuringTick)
+	{
+		UMockAI_BT::ExecutionLog.Add(LogIndex);
+		OwnerComp.StopTree();
+	}
+}
\ No newline at end of file
diff --git a/Engine/Source/Developer/AITestSuite/Private/BehaviorTree/TestBTTask_StopTree.cpp b/Engine/Source/Developer/AITestSuite/Private/BehaviorTree/TestBTTask_StopTree.cpp
new file mode 100644
--- /dev/null
+++ b/Engine/Source/Developer/AITestSuite/Private/BehaviorTree/TestBTTask_StopTree.cpp
@@ -0,0 +1,53 @@
+// Copyright Epic Games, Inc. All Rights Reserved.
+
+#include "BehaviorTree/TestBTTask_StopTree.h"
+#include "MockAI_BT.h"
+
+UTestBTTask_StopTree::UTestBTTask_StopTree(const FObjectInitializer& ObjectInitializer) : Super(ObjectInitializer)
+{
+	StopTimming = EBTTestTaskStopTree::DuringExecute;
+	LogResult = EBTNodeResult::Succeeded;
+	LogIndex = INDEX_NONE;
+	INIT_TASK_NODE_NOTIFY_FLAGS();
+}
+
+EBTNodeResult::Type UTestBTTask_StopTree::ExecuteTask(UBehaviorTreeComponent& OwnerComp, uint8* NodeMemory)
+{
+	if (StopTimming == EBTTestTaskStopTree::DuringExecute)
+	{
+		UMockAI_BT::ExecutionLog.Add(LogIndex);
+		OwnerComp.StopTree();
+		return LogResult;
+	}
+
+	return EBTNodeResult::InProgress;
+}
+
+EBTNodeResult::Type UTestBTTask_StopTree::AbortTask(UBehaviorTreeComponent& OwnerComp, uint8* NodeMemory)
+{
+	if (StopTimming == EBTTestTaskStopTree::DuringAbort)
+	{
+		UMockAI_BT::ExecutionLog.Add(LogIndex);
+		OwnerComp.StopTree();
+	}
+	return EBTNodeResult::Aborted;
+}
+
+void UTestBTTask_StopTree::TickTask(class UBehaviorTreeComponent& OwnerComp, uint8* NodeMemory, float DeltaSeconds)
+{
+	if (StopTimming == EBTTestTaskStopTree::DuringTick)
+	{
+		UMockAI_BT::ExecutionLog.Add(LogIndex);
+		OwnerComp.StopTree();
+		FinishLatentTask(OwnerComp, LogResult);
+	}
+}
+
+void UTestBTTask_StopTree::OnTaskFinished(UBehaviorTreeComponent& OwnerComp, uint8* NodeMemory, EBTNodeResult::Type TaskResult)
+{
+	if (StopTimming == EBTTestTaskStopTree::DuringFinish)
+	{
+		UMockAI_BT::ExecutionLog.Add(LogIndex);
+		OwnerComp.StopTree();
+	}
+}
diff --git a/Engine/Source/Developer/AITestSuite/Private/Tests/BTTest.cpp b/Engine/Source/Developer/AITestSuite/Private/Tests/BTTest.cpp
--- a/Engine/Source/Developer/AITestSuite/Private/Tests/BTTest.cpp
+++ b/Engine/Source/Developer/AITestSuite/Private/Tests/BTTest.cpp
@@ -2273,3 +2273,323 @@ struct FAITest_BTPostponeDuringSearch2 : public FAITest_SimpleBT
 };
 IMPLEMENT_AI_LATENT_TEST(FAITest_BTPostponeDuringSearch2, "System.AI.Behavior Trees.Service: Postpone during search (lower pri)")
 #undef LOCTEXT_NAMESPACE
+
+struct FAITest_BTStopDuringTaskExecute : public FAITest_SimpleBT
+{
+	FAITest_BTStopDuringTaskExecute()
+	{
+		UBTCompositeNode& RootNode = FBTBuilder::AddSelector(*BTAsset);
+		UBTCompositeNode& CompNode = FBTBuilder::AddSelector(RootNode);
+		{
+			FBTBuilder::WithServiceLog(CompNode, 1/*ActivationIndex*/, 2/*DeactivationIndex*/, 3/*TickIndex*/);
+
+			UBTCompositeNode& Comp1Node = FBTBuilder::AddSelector(CompNode);
+			{
+				FBTBuilder::AddTaskStopTree(Comp1Node, 4, EBTNodeResult::Succeeded, EBTTestTaskStopTree::DuringExecute);
+			}
+		}
+
+		ExpectedResult.Add(1);
+		ExpectedResult.Add(3);
+		ExpectedResult.Add(4);
+		ExpectedResult.Add(2);
+	}
+};
+IMPLEMENT_AI_LATENT_TEST(FAITest_BTStopDuringTaskExecute, "System.AI.Behavior Trees.Stop: during task execute")
+
+struct FAITest_BTStopDuringTaskTick : public FAITest_SimpleBT
+{
+	FAITest_BTStopDuringTaskTick()
+	{
+		UBTCompositeNode& RootNode = FBTBuilder::AddSelector(*BTAsset);
+		UBTCompositeNode& CompNode = FBTBuilder::AddSelector(RootNode);
+		{
+			FBTBuilder::WithServiceLog(CompNode, 1/*ActivationIndex*/, 2/*DeactivationIndex*/, 3/*TickIndex*/);
+
+			UBTCompositeNode& Comp1Node = FBTBuilder::AddSelector(CompNode);
+			{
+				FBTBuilder::AddTaskStopTree(Comp1Node, 4, EBTNodeResult::Succeeded, EBTTestTaskStopTree::DuringTick);
+			}
+		}
+
+		ExpectedResult.Add(1);
+		ExpectedResult.Add(3);
+		ExpectedResult.Add(4);
+		ExpectedResult.Add(2);
+	}
+};
+IMPLEMENT_AI_LATENT_TEST(FAITest_BTStopDuringTaskTick, "System.AI.Behavior Trees.Stop: during task tick")
+
+
+struct FAITest_BTStopDuringTaskAbort : public FAITest_SimpleBT
+{
+	FAITest_BTStopDuringTaskAbort()
+	{
+		UBTCompositeNode& RootNode = FBTBuilder::AddSelector(*BTAsset);
+		UBTCompositeNode& CompNode = FBTBuilder::AddSelector(RootNode);
+		{
+			FBTBuilder::WithServiceLog(CompNode, 1/*ActivationIndex*/, 2/*DeactivationIndex*/, 3/*TickIndex*/);
+
+			UBTCompositeNode& Comp1Node = FBTBuilder::AddSelector(CompNode);
+			{
+				FBTBuilder::WithServiceLog(Comp1Node, 4/*ActivationIndex*/, 5/*DeactivationIndex*/, 6/*TickIndex*/, TEXT("Bool1"));
+
+				FBTBuilder::AddTaskStopTree(Comp1Node, 7, EBTNodeResult::Succeeded, EBTTestTaskStopTree::DuringAbort);
+				FBTBuilder::WithDecoratorBlackboard(Comp1Node, EBasicKeyOperation::NotSet, EBTFlowAbortMode::Both, TEXT("Bool1"));
+
+				FBTBuilder::AddTask(Comp1Node, 8, EBTNodeResult::Succeeded);
+			}
+		}
+
+		ExpectedResult.Add(1);
+		ExpectedResult.Add(4);
+		ExpectedResult.Add(3);
+		ExpectedResult.Add(6);
+		ExpectedResult.Add(3);
+		ExpectedResult.Add(6);
+		ExpectedResult.Add(7);
+		ExpectedResult.Add(2);
+		ExpectedResult.Add(5);
+	}
+};
+IMPLEMENT_AI_LATENT_TEST(FAITest_BTStopDuringTaskAbort, "System.AI.Behavior Trees.Stop: during task abort")
+
+
+struct FAITest_BTStopDuringTaskFinished : public FAITest_SimpleBT
+{
+	FAITest_BTStopDuringTaskFinished()
+	{
+		UBTCompositeNode& RootNode = FBTBuilder::AddSelector(*BTAsset);
+		UBTCompositeNode& CompNode = FBTBuilder::AddSelector(RootNode);
+		{
+			FBTBuilder::WithServiceLog(CompNode, 1/*ActivationIndex*/, 2/*DeactivationIndex*/, 3/*TickIndex*/);
+
+			UBTCompositeNode& Comp1Node = FBTBuilder::AddSelector(CompNode);
+			{
+				FBTBuilder::WithServiceLog(Comp1Node, 4/*ActivationIndex*/, 5/*DeactivationIndex*/, 6/*TickIndex*/, TEXT("Bool1"));
+
+				FBTBuilder::AddTaskStopTree(Comp1Node, 7, EBTNodeResult::Succeeded, EBTTestTaskStopTree::DuringFinish);
+				FBTBuilder::WithDecoratorBlackboard(Comp1Node, EBasicKeyOperation::NotSet, EBTFlowAbortMode::Both, TEXT("Bool1"));
+
+				FBTBuilder::AddTask(Comp1Node, 8, EBTNodeResult::Succeeded);
+			}
+		}
+
+		ExpectedResult.Add(1);
+		ExpectedResult.Add(4);
+		ExpectedResult.Add(3);
+		ExpectedResult.Add(6);
+		ExpectedResult.Add(3);
+		ExpectedResult.Add(6);
+		ExpectedResult.Add(7);
+		ExpectedResult.Add(2);
+		ExpectedResult.Add(5);
+	}
+};
+IMPLEMENT_AI_LATENT_TEST(FAITest_BTStopDuringTaskFinished, "System.AI.Behavior Trees.Stop: during task finished")
+
+struct FAITest_BTStopDuringServiceTick : public FAITest_SimpleBT
+{
+	FAITest_BTStopDuringServiceTick()
+	{
+		UBTCompositeNode& RootNode = FBTBuilder::AddSelector(*BTAsset);
+		UBTCompositeNode& CompNode = FBTBuilder::AddSelector(RootNode);
+		{
+			FBTBuilder::WithServiceLog(CompNode, 1/*ActivationIndex*/, 2/*DeactivationIndex*/, 3/*TickIndex*/);
+
+			UBTCompositeNode& Comp1Node = FBTBuilder::AddSelector(CompNode);
+			{
+				FBTBuilder::WithServiceLog(Comp1Node, 4/*ActivationIndex*/, 5/*DeactivationIndex*/, 6/*TickIndex*/, TEXT("Bool1"));
+				FBTBuilder::WithServiceStopTree(Comp1Node, 9, EBTTestServiceStopTree::DuringTick);
+
+				FBTBuilder::AddTask(Comp1Node, 7, EBTNodeResult::Succeeded);
+				FBTBuilder::WithDecoratorBlackboard(Comp1Node, EBasicKeyOperation::NotSet, EBTFlowAbortMode::Both, TEXT("Bool1"));
+
+				FBTBuilder::AddTask(Comp1Node, 8, EBTNodeResult::Succeeded);
+			}
+		}
+
+		ExpectedResult.Add(1);
+		ExpectedResult.Add(4);
+		ExpectedResult.Add(3);
+		ExpectedResult.Add(6);
+		ExpectedResult.Add(9);
+		ExpectedResult.Add(7);
+		ExpectedResult.Add(2);
+		ExpectedResult.Add(5);
+	}
+};
+IMPLEMENT_AI_LATENT_TEST(FAITest_BTStopDuringServiceTick, "System.AI.Behavior Trees.Stop: during service tick")
+
+struct FAITest_BTStopDuringServiceBecomeRelevant : public FAITest_SimpleBT
+{
+	FAITest_BTStopDuringServiceBecomeRelevant()
+	{
+		UBTCompositeNode& RootNode = FBTBuilder::AddSelector(*BTAsset);
+		UBTCompositeNode& CompNode = FBTBuilder::AddSelector(RootNode);
+		{
+			FBTBuilder::WithServiceLog(CompNode, 1/*ActivationIndex*/, 2/*DeactivationIndex*/, 3/*TickIndex*/);
+
+			UBTCompositeNode& Comp1Node = FBTBuilder::AddSelector(CompNode);
+			{
+				FBTBuilder::WithServiceLog(Comp1Node, 4/*ActivationIndex*/, 5/*DeactivationIndex*/, 6/*TickIndex*/, TEXT("Bool1"));
+				FBTBuilder::WithServiceStopTree(Comp1Node, 9, EBTTestServiceStopTree::DuringBecomeRelevant);
+
+				FBTBuilder::AddTask(Comp1Node, 7, EBTNodeResult::Succeeded);
+				FBTBuilder::WithDecoratorBlackboard(Comp1Node, EBasicKeyOperation::NotSet, EBTFlowAbortMode::Both, TEXT("Bool1"));
+
+				FBTBuilder::AddTask(Comp1Node, 8, EBTNodeResult::Succeeded);
+			}
+		}
+
+		ExpectedResult.Add(1);
+		ExpectedResult.Add(4);
+		ExpectedResult.Add(9);
+		ExpectedResult.Add(3);
+		ExpectedResult.Add(6);
+		ExpectedResult.Add(7);
+		ExpectedResult.Add(2);
+		ExpectedResult.Add(5);
+	}
+};
+IMPLEMENT_AI_LATENT_TEST(FAITest_BTStopDuringServiceBecomeRelevant, "System.AI.Behavior Trees.Stop: during service become relevant")
+
+struct FAITest_BTStopDuringServiceCeaseRelevant : public FAITest_SimpleBT
+{
+	FAITest_BTStopDuringServiceCeaseRelevant()
+	{
+		UBTCompositeNode& RootNode = FBTBuilder::AddSelector(*BTAsset);
+		UBTCompositeNode& CompNode = FBTBuilder::AddSelector(RootNode);
+		{
+			FBTBuilder::WithServiceLog(CompNode, 1/*ActivationIndex*/, 2/*DeactivationIndex*/, 3/*TickIndex*/);
+
+			UBTCompositeNode& Comp1Node = FBTBuilder::AddSelector(CompNode);
+			{
+				FBTBuilder::WithServiceLog(Comp1Node, 4/*ActivationIndex*/, 5/*DeactivationIndex*/, 6/*TickIndex*/, TEXT("Bool1"));
+				FBTBuilder::WithServiceStopTree(Comp1Node, 9, EBTTestServiceStopTree::DuringCeaseRelevant);
+
+				FBTBuilder::AddTask(Comp1Node, 7, EBTNodeResult::Succeeded);
+				FBTBuilder::WithDecoratorBlackboard(Comp1Node, EBasicKeyOperation::NotSet, EBTFlowAbortMode::Both, TEXT("Bool1"));
+			}
+			FBTBuilder::AddTask(CompNode, 8, EBTNodeResult::Succeeded);
+		}
+
+		ExpectedResult.Add(1);
+		ExpectedResult.Add(4);
+		ExpectedResult.Add(3);
+		ExpectedResult.Add(6);
+		ExpectedResult.Add(7);
+		ExpectedResult.Add(3);
+		ExpectedResult.Add(6);
+		ExpectedResult.Add(5);
+		ExpectedResult.Add(9);
+		ExpectedResult.Add(8);
+		ExpectedResult.Add(2);
+	}
+};
+IMPLEMENT_AI_LATENT_TEST(FAITest_BTStopDuringServiceCeaseRelevant, "System.AI.Behavior Trees.Stop: during service cease relevant")
+
+struct FAITest_BTStopDuringTaskServiceTick : public FAITest_SimpleBT
+{
+	FAITest_BTStopDuringTaskServiceTick()
+	{
+		UBTCompositeNode& RootNode = FBTBuilder::AddSelector(*BTAsset);
+		UBTCompositeNode& CompNode = FBTBuilder::AddSelector(RootNode);
+		{
+			FBTBuilder::WithServiceLog(CompNode, 1/*ActivationIndex*/, 2/*DeactivationIndex*/, 3/*TickIndex*/);
+
+			UBTCompositeNode& Comp1Node = FBTBuilder::AddSelector(CompNode);
+			{
+				FBTBuilder::WithServiceLog(Comp1Node, 4/*ActivationIndex*/, 5/*DeactivationIndex*/, 6/*TickIndex*/, TEXT("Bool1"));
+
+				FBTBuilder::AddTask(Comp1Node, 7, EBTNodeResult::Succeeded);
+				FBTBuilder::WithDecoratorBlackboard(Comp1Node, EBasicKeyOperation::NotSet, EBTFlowAbortMode::Both, TEXT("Bool1"));
+				FBTBuilder::WithTaskServiceStopTree(Comp1Node, 9, EBTTestServiceStopTree::DuringTick);
+
+				FBTBuilder::AddTask(Comp1Node, 8, EBTNodeResult::Succeeded);
+			}
+		}
+
+		ExpectedResult.Add(1);
+		ExpectedResult.Add(4);
+		ExpectedResult.Add(3);
+		ExpectedResult.Add(6);
+		ExpectedResult.Add(9);
+		ExpectedResult.Add(7);
+		ExpectedResult.Add(2);
+		ExpectedResult.Add(5);
+	}
+};
+IMPLEMENT_AI_LATENT_TEST(FAITest_BTStopDuringTaskServiceTick, "System.AI.Behavior Trees.Stop: during task service tick")
+
+struct FAITest_BTStopDuringTaskServiceBecomeRelevant : public FAITest_SimpleBT
+{
+	FAITest_BTStopDuringTaskServiceBecomeRelevant()
+	{
+		UBTCompositeNode& RootNode = FBTBuilder::AddSelector(*BTAsset);
+		UBTCompositeNode& CompNode = FBTBuilder::AddSelector(RootNode);
+		{
+			FBTBuilder::WithServiceLog(CompNode, 1/*ActivationIndex*/, 2/*DeactivationIndex*/, 3/*TickIndex*/);
+
+			UBTCompositeNode& Comp1Node = FBTBuilder::AddSelector(CompNode);
+			{
+				FBTBuilder::WithServiceLog(Comp1Node, 4/*ActivationIndex*/, 5/*DeactivationIndex*/, 6/*TickIndex*/, TEXT("Bool1"));
+
+				FBTBuilder::AddTask(Comp1Node, 7, EBTNodeResult::Succeeded);
+				FBTBuilder::WithDecoratorBlackboard(Comp1Node, EBasicKeyOperation::NotSet, EBTFlowAbortMode::Both, TEXT("Bool1"));
+				FBTBuilder::WithTaskServiceStopTree(Comp1Node, 9, EBTTestServiceStopTree::DuringBecomeRelevant);
+
+				FBTBuilder::AddTask(Comp1Node, 8, EBTNodeResult::Succeeded);
+			}
+		}
+
+		ExpectedResult.Add(1);
+		ExpectedResult.Add(4);
+		ExpectedResult.Add(3);
+		ExpectedResult.Add(6);
+		ExpectedResult.Add(9);
+		ExpectedResult.Add(7);
+		ExpectedResult.Add(2);
+		ExpectedResult.Add(5);
+	}
+};
+IMPLEMENT_AI_LATENT_TEST(FAITest_BTStopDuringTaskServiceBecomeRelevant, "System.AI.Behavior Trees.Stop: during task service become relevant")
+
+
+struct FAITest_BTStopDuringTaskServiceCeaseRelevant : public FAITest_SimpleBT
+{
+	FAITest_BTStopDuringTaskServiceCeaseRelevant()
+	{
+		UBTCompositeNode& RootNode = FBTBuilder::AddSelector(*BTAsset);
+		UBTCompositeNode& CompNode = FBTBuilder::AddSelector(RootNode);
+		{
+			FBTBuilder::WithServiceLog(CompNode, 1/*ActivationIndex*/, 2/*DeactivationIndex*/, 3/*TickIndex*/);
+
+			UBTCompositeNode& Comp1Node = FBTBuilder::AddSelector(CompNode);
+			{
+				FBTBuilder::WithServiceLog(Comp1Node, 4/*ActivationIndex*/, 5/*DeactivationIndex*/, 6/*TickIndex*/, TEXT("Bool1"));
+
+				FBTBuilder::AddTask(Comp1Node, 7, EBTNodeResult::Succeeded);
+				FBTBuilder::WithDecoratorBlackboard(Comp1Node, EBasicKeyOperation::NotSet, EBTFlowAbortMode::Both, TEXT("Bool1"));
+				FBTBuilder::WithTaskServiceStopTree(Comp1Node, 9, EBTTestServiceStopTree::DuringCeaseRelevant);
+
+				FBTBuilder::AddTask(CompNode, 8, EBTNodeResult::Succeeded);
+			}
+		}
+
+		ExpectedResult.Add(1);
+		ExpectedResult.Add(4);
+		ExpectedResult.Add(3);
+		ExpectedResult.Add(6);
+		ExpectedResult.Add(7);
+		ExpectedResult.Add(3);
+		ExpectedResult.Add(6);
+		ExpectedResult.Add(9);
+		ExpectedResult.Add(5);
+		ExpectedResult.Add(8);
+		ExpectedResult.Add(2);
+	}
+};
+IMPLEMENT_AI_LATENT_TEST(FAITest_BTStopDuringTaskServiceCeaseRelevant, "System.AI.Behavior Trees.Stop: during task service cease relevant")
+
+#undef LOCTEXT_NAMESPACE
\ No newline at end of file
diff --git a/Engine/Source/Developer/AITestSuite/Public/BTBuilder.h b/Engine/Source/Developer/AITestSuite/Public/BTBuilder.h
--- a/Engine/Source/Developer/AITestSuite/Public/BTBuilder.h
+++ b/Engine/Source/Developer/AITestSuite/Public/BTBuilder.h
@@ -15,14 +15,16 @@
 #include "BehaviorTree/Decorators/BTDecorator_Loop.h"
 #include "BehaviorTree/Tasks/BTTask_RunBehavior.h"
 
-#include "BehaviorTree/TestBTTask_Log.h"
-#include "BehaviorTree/TestBTTask_SetFlag.h"
-#include "BehaviorTree/TestBTTask_ToggleFlag.h"
-#include "BehaviorTree/TestBTTask_SetValue.h"
-#include "BehaviorTree/TestBTTask_LatentWithFlags.h"
 #include "BehaviorTree/TestBTDecorator_Blackboard.h"
 #include "BehaviorTree/TestBTDecorator_DelayedAbort.h"
 #include "BehaviorTree/TestBTService_Log.h"
+#include "BehaviorTree/TestBTService_StopTree.h"
+#include "BehaviorTree/TestBTTask_LatentWithFlags.h"
+#include "BehaviorTree/TestBTTask_Log.h"
+#include "BehaviorTree/TestBTTask_SetFlag.h"
+#include "BehaviorTree/TestBTTask_SetValue.h"
+#include "BehaviorTree/TestBTTask_StopTree.h"
+#include "BehaviorTree/TestBTTask_ToggleFlag.h"
 
 struct FBTBuilder
 {
@@ -219,6 +221,17 @@ struct FBTBuilder
 		ParentNode.Children[ChildIdx].ChildTask = TaskNode;
 	}
 
+	static void AddTaskStopTree(UBTCompositeNode& ParentNode, int32 LogIndex, EBTNodeResult::Type NodeResult, EBTTestTaskStopTree::Type StopTimming)
+	{
+		UTestBTTask_StopTree* TaskNode = NewObject<UTestBTTask_StopTree>(ParentNode.GetTreeAsset());
+		TaskNode->LogIndex = LogIndex;
+		TaskNode->LogResult = NodeResult;
+		TaskNode->StopTimming = StopTimming;
+
+		const int32 ChildIdx = ParentNode.Children.AddZeroed(1);
+		ParentNode.Children[ChildIdx].ChildTask = TaskNode;
+	}
+
 	template<class T>
 	static T& WithDecorator(UBTCompositeNode& ParentNode, UClass* DecoratorClass = T::StaticClass())
 	{
@@ -314,6 +327,13 @@ struct FBTBuilder
 		LogService.bToggleValue = bToggleValue;
 	}
 
+	static void WithServiceStopTree(UBTCompositeNode& ParentNode, int32 LogIndex, EBTTestServiceStopTree::Type StopTimming)
+	{
+		UTestBTService_StopTree& Service = WithService<UTestBTService_StopTree>(ParentNode);
+		Service.LogIndex = LogIndex;
+		Service.StopTimming = StopTimming;
+	}
+
 	template<class T>
 	static T& WithTaskService(UBTCompositeNode& ParentNode, UClass* ServiceClass = T::StaticClass())
 	{
@@ -337,4 +357,12 @@ struct FBTBuilder
 		LogService.KeyNameCeaseRelevant = CeaseRelevantBoolKeyName;
 		LogService.bToggleValue = bToggleValue;
 	}
+
+	static void WithTaskServiceStopTree(UBTCompositeNode& ParentNode, int32 LogIndex, EBTTestServiceStopTree::Type StopTimming)
+	{
+		UTestBTService_StopTree& Service = WithTaskService<UTestBTService_StopTree>(ParentNode);
+		Service.LogIndex = LogIndex;
+		Service.StopTimming = StopTimming;
+	}
+
 };
diff --git a/Engine/Source/Runtime/AIModule/Classes/BehaviorTree/BehaviorTreeComponent.h b/Engine/Source/Runtime/AIModule/Classes/BehaviorTree/BehaviorTreeComponent.h
--- a/Engine/Source/Runtime/AIModule/Classes/BehaviorTree/BehaviorTreeComponent.h
+++ b/Engine/Source/Runtime/AIModule/Classes/BehaviorTree/BehaviorTreeComponent.h
@@ -21,7 +21,6 @@ class UBTNode;
 class UBTTask_RunBehavior;
 class UBTTask_RunBehaviorDynamic;
 class UBTTaskNode;
-struct FScopedBehaviorTreeLock;
 
 struct FBTNodeExecutionInfo
 {
@@ -88,7 +87,10 @@ enum class EBTBranchAction : uint8
 	DecoratorActivate = DecoratorActivate_IfNotExecuting | DecoratorActivate_EvenIfExecuting,
 	DecoratorDeactivate = 0x8,
 	UnregisterAuxNodes = 0x10,
-	All = DecoratorEvaluate | DecoratorActivate_IfNotExecuting | DecoratorActivate_EvenIfExecuting | DecoratorDeactivate | UnregisterAuxNodes,
+	StopTree_Safe = 0x20,
+	StopTree_Forced = 0x40,
+	StopTree = StopTree_Safe | StopTree_Forced,
+	All = DecoratorEvaluate | DecoratorActivate_IfNotExecuting | DecoratorActivate_EvenIfExecuting | DecoratorDeactivate | UnregisterAuxNodes | StopTree_Safe | StopTree_Forced,
 };
 ENUM_CLASS_FLAGS(EBTBranchAction);
 
@@ -336,12 +338,6 @@ protected:
 	/** index of last active instance on stack */
 	uint16 ActiveInstanceIdx;
 
-	/** if set, StopTree calls will be deferred */
-	uint8 StopTreeLock;
-
-	/** if set, StopTree will be called at the end of tick */
-	uint8 bDeferredStopTree : 1;
-
 	/** loops tree execution */
 	uint8 bLoopExecution : 1;
 
@@ -492,7 +488,6 @@ protected:
 	friend UBTTask_RunBehaviorDynamic;
 	friend FBehaviorTreeDebugger;
 	friend FBehaviorTreeInstance;
-	friend FScopedBehaviorTreeLock;
 
 protected:
 	/** data asset defining the tree */
diff --git a/Engine/Source/Runtime/AIModule/Private/BehaviorTree/BehaviorTreeComponent.cpp b/Engine/Source/Runtime/AIModule/Private/BehaviorTree/BehaviorTreeComponent.cpp
--- a/Engine/Source/Runtime/AIModule/Private/BehaviorTree/BehaviorTreeComponent.cpp
+++ b/Engine/Source/Runtime/AIModule/Private/BehaviorTree/BehaviorTreeComponent.cpp
@@ -30,29 +30,6 @@ double UBehaviorTreeComponent::FrameSearchTime = 0.;
 int32 UBehaviorTreeComponent::NumSearchTimeCalls = 0;
 #endif
 
-struct FScopedBehaviorTreeLock
-{
-	FScopedBehaviorTreeLock(UBehaviorTreeComponent& InOwnerComp, uint8 InLockFlag) : OwnerComp(InOwnerComp), LockFlag(InLockFlag)
-	{
-		OwnerComp.StopTreeLock |= LockFlag;
-	}
-
-	~FScopedBehaviorTreeLock()
-	{
-		OwnerComp.StopTreeLock &= ~LockFlag;
-	}
-
-	enum
-	{
-		LockTick = 1 << 0,
-		LockReentry = 1 << 1,
-	};
-
-private:
-	UBehaviorTreeComponent& OwnerComp;
-	uint8 LockFlag;
-};
-
 //----------------------------------------------------------------------//
 // UBehaviorTreeComponent
 //----------------------------------------------------------------------//
@@ -62,8 +39,6 @@ UBehaviorTreeComponent::UBehaviorTreeComponent(const FObjectInitializer& ObjectI
 	, SearchData(*this)
 {
 	ActiveInstanceIdx = 0;
-	StopTreeLock = 0;
-	bDeferredStopTree = false;
 	bLoopExecution = false;
 	bWaitingForLatentAborts = false;
 	bRequestedFlowUpdate = false;
@@ -285,16 +260,16 @@ void UBehaviorTreeComponent::StopTree(EBTStopMode::Type StopMode)
 {
 	SCOPE_CYCLE_COUNTER(STAT_AI_BehaviorTree_StopTree);
 
-	UE_VLOG(GetOwner(), LogBehaviorTree, Verbose, TEXT("StopTree %s, mode:%s"), *GetNameSafe(GetRootTree()), StopMode == EBTStopMode::Forced ? TEXT("Forced") : TEXT("Safe"));
-
-	if (StopTreeLock)
+	const bool bForcedStop = StopMode == EBTStopMode::Forced;
+	if ((SuspendedBranchActions & EBTBranchAction::StopTree) != EBTBranchAction::None)
 	{
-		bDeferredStopTree = true;
-		ScheduleNextTick(0.0f);
+		UE_VLOG(GetOwner(), LogBehaviorTree, Verbose, TEXT("Stop tree(%s) queued up"), bForcedStop ? TEXT("Forced") : TEXT("Safe"));
+		PendingBranchActionRequests.Emplace(nullptr, bForcedStop ? EBTBranchAction::StopTree_Forced : EBTBranchAction::StopTree_Safe);
 		return;
 	}
 
-	FScopedBehaviorTreeLock ScopedLock(*this, FScopedBehaviorTreeLock::LockReentry);
+	UE_VLOG(GetOwner(), LogBehaviorTree, Verbose, TEXT("StopTree %s, mode:%s"), *GetNameSafe(GetRootTree()), bForcedStop ? TEXT("Forced") : TEXT("Safe"));
+	FBTSuspendBranchActionsScoped ScopedSuspend(*this, EBTBranchAction::All);
 	if (!bRequestedStop)
 	{
 		bRequestedStop = true;
@@ -305,7 +280,6 @@ void UBehaviorTreeComponent::StopTree(EBTStopMode::Type StopMode)
 
 			// notify active aux nodes
 			{
-				FBTSuspendBranchActionsScoped ScopedSuspend(*this, EBTBranchAction::All);
 				InstanceInfo.ExecuteOnEachAuxNode([&InstanceInfo, this](const UBTAuxiliaryNode& AuxNode)
 				{
 					uint8* NodeMemory = AuxNode.GetNodeMemory<uint8>(InstanceInfo);
@@ -388,7 +362,7 @@ void UBehaviorTreeComponent::StopTree(EBTStopMode::Type StopMode)
 
 	if (bWaitingForLatentAborts)
 	{
-		if (StopMode == EBTStopMode::Safe)
+		if (!bForcedStop)
 		{
 			UE_VLOG(GetOwner(), LogBehaviorTree, Log, TEXT("StopTree is waiting for aborting tasks to finish..."));
 			return;
@@ -424,7 +398,6 @@ void UBehaviorTreeComponent::StopTree(EBTStopMode::Type StopMode)
 	bRequestedStop = false;
 	bIsRunning = false;
 	bWaitingForLatentAborts = false;
-	bDeferredStopTree = false;
 }
 
 void UBehaviorTreeComponent::RestartTree()
@@ -793,6 +766,11 @@ void UBehaviorTreeComponent::ResumeBranchActions()
 					UnregisterAuxNodesInBranch(BranchRoot, true/*bApplyImmediately*/);
 					break;
 				}
+				case EBTBranchAction::StopTree_Safe:
+				case EBTBranchAction::StopTree_Forced:
+				{
+					StopTree(Info.Action == EBTBranchAction::StopTree_Forced ? EBTStopMode::Forced : EBTStopMode::Safe);
+				}
 			}
 		}
 	}
@@ -1556,7 +1534,7 @@ void UBehaviorTreeComponent::TickComponent(float DeltaTime, enum ELevelTick Tick
 	// do it before processing execution request to give BP driven logic chance to accumulate execution requests
 	// newly added aux nodes are ticked as part of SearchData application
 	{
-		FBTSuspendBranchActionsScoped ScopedSuspend(*this, EBTBranchAction::UnregisterAuxNodes);
+		FBTSuspendBranchActionsScoped ScopedSuspend(*this, EBTBranchAction::UnregisterAuxNodes | EBTBranchAction::StopTree);
 		for (int32 InstanceIndex = 0; InstanceIndex < InstanceStack.Num(); InstanceIndex++)
 		{
 			FBehaviorTreeInstance& InstanceInfo = InstanceStack[InstanceIndex];
@@ -1609,54 +1587,46 @@ void UBehaviorTreeComponent::TickComponent(float DeltaTime, enum ELevelTick Tick
 
 	if (InstanceStack.Num() > 0 && bIsRunning && !bIsPaused)
 	{
-		{
-			FScopedBehaviorTreeLock ScopedLock(*this, FScopedBehaviorTreeLock::LockTick);
-
-			// tick active parallel tasks (in execution order, before task)
-			for (int32 InstanceIndex = 0; InstanceIndex < InstanceStack.Num(); InstanceIndex++)
-			{
-				FBehaviorTreeInstance& InstanceInfo = InstanceStack[InstanceIndex];
-				InstanceInfo.ExecuteOnEachParallelTask([&InstanceInfo, &bDoneSomething, this, DeltaTime, &NextNeededDeltaTime](const FBehaviorTreeParallelTask& ParallelTaskInfo, const int32 Index)
-					{
-						const UBTTaskNode* ParallelTask = ParallelTaskInfo.TaskNode;
-						SCOPE_CYCLE_UOBJECT(ParallelTask, ParallelTask);
-						uint8* NodeMemory = ParallelTask->GetNodeMemory<uint8>(InstanceInfo);
-						bDoneSomething |= ParallelTask->WrappedTickTask(*this, NodeMemory, DeltaTime, NextNeededDeltaTime);
-					});
-			}
+		FBTSuspendBranchActionsScoped ScopedSuspend(*this, EBTBranchAction::StopTree);
 
-			// tick active task
-			if (InstanceStack.IsValidIndex(ActiveInstanceIdx))
-			{
-				FBehaviorTreeInstance& ActiveInstance = InstanceStack[ActiveInstanceIdx];
-				if (ActiveInstance.ActiveNodeType == EBTActiveNode::ActiveTask ||
-					ActiveInstance.ActiveNodeType == EBTActiveNode::AbortingTask)
+		// tick active parallel tasks (in execution order, before task)
+		for (int32 InstanceIndex = 0; InstanceIndex < InstanceStack.Num(); InstanceIndex++)
+		{
+			FBehaviorTreeInstance& InstanceInfo = InstanceStack[InstanceIndex];
+			InstanceInfo.ExecuteOnEachParallelTask([&InstanceInfo, &bDoneSomething, this, DeltaTime, &NextNeededDeltaTime](const FBehaviorTreeParallelTask& ParallelTaskInfo, const int32 Index)
 				{
-					UBTTaskNode* ActiveTask = (UBTTaskNode*)ActiveInstance.ActiveNode;
-					uint8* NodeMemory = ActiveTask->GetNodeMemory<uint8>(ActiveInstance);
-					SCOPE_CYCLE_UOBJECT(ActiveTask, ActiveTask);
-					bDoneSomething |= ActiveTask->WrappedTickTask(*this, NodeMemory, DeltaTime, NextNeededDeltaTime);
-				}
-			}
+					const UBTTaskNode* ParallelTask = ParallelTaskInfo.TaskNode;
+					SCOPE_CYCLE_UOBJECT(ParallelTask, ParallelTask);
+					uint8* NodeMemory = ParallelTask->GetNodeMemory<uint8>(InstanceInfo);
+					bDoneSomething |= ParallelTask->WrappedTickTask(*this, NodeMemory, DeltaTime, NextNeededDeltaTime);
+				});
+		}
 
-			// tick aborting task from abandoned subtree
-			if (InstanceStack.IsValidIndex(ActiveInstanceIdx + 1))
+		// tick active task
+		if (InstanceStack.IsValidIndex(ActiveInstanceIdx))
+		{
+			FBehaviorTreeInstance& ActiveInstance = InstanceStack[ActiveInstanceIdx];
+			if (ActiveInstance.ActiveNodeType == EBTActiveNode::ActiveTask ||
+				ActiveInstance.ActiveNodeType == EBTActiveNode::AbortingTask)
 			{
-				FBehaviorTreeInstance& LastInstance = InstanceStack.Last();
-				if (LastInstance.ActiveNodeType == EBTActiveNode::AbortingTask)
-				{
-					UBTTaskNode* ActiveTask = (UBTTaskNode*)LastInstance.ActiveNode;
-					uint8* NodeMemory = ActiveTask->GetNodeMemory<uint8>(LastInstance);
-					SCOPE_CYCLE_UOBJECT(ActiveTask, ActiveTask);
-					bDoneSomething |= ActiveTask->WrappedTickTask(*this, NodeMemory, DeltaTime, NextNeededDeltaTime);
-				}
+				UBTTaskNode* ActiveTask = (UBTTaskNode*)ActiveInstance.ActiveNode;
+				uint8* NodeMemory = ActiveTask->GetNodeMemory<uint8>(ActiveInstance);
+				SCOPE_CYCLE_UOBJECT(ActiveTask, ActiveTask);
+				bDoneSomething |= ActiveTask->WrappedTickTask(*this, NodeMemory, DeltaTime, NextNeededDeltaTime);
 			}
 		}
 
-		if (bDeferredStopTree)
+		// tick aborting task from abandoned subtree
+		if (InstanceStack.IsValidIndex(ActiveInstanceIdx + 1))
 		{
-			StopTree(EBTStopMode::Safe);
-			bDoneSomething = true;
+			FBehaviorTreeInstance& LastInstance = InstanceStack.Last();
+			if (LastInstance.ActiveNodeType == EBTActiveNode::AbortingTask)
+			{
+				UBTTaskNode* ActiveTask = (UBTTaskNode*)LastInstance.ActiveNode;
+				uint8* NodeMemory = ActiveTask->GetNodeMemory<uint8>(LastInstance);
+				SCOPE_CYCLE_UOBJECT(ActiveTask, ActiveTask);
+				bDoneSomething |= ActiveTask->WrappedTickTask(*this, NodeMemory, DeltaTime, NextNeededDeltaTime);
+			}
 		}
 	}
 
@@ -2271,10 +2241,6 @@ void UBehaviorTreeComponent::ExecuteTask(UBTTaskNode* TaskNode)
 	ActiveInstance.ActiveNode = TaskNode;
 	ActiveInstance.ActiveNodeType = EBTActiveNode::ActiveTask;
 
-	// Is is now ok to resume any branch acitons as the new active instance is set.
-	// Before that, decorators evaluating the IsExecutingBranch would be wrong.
-	ResumeBranchActions();
-
 	// make a snapshot for debugger
 	StoreDebuggerExecutionStep(EBTExecutionSnap::Regular);
 
@@ -2299,6 +2265,10 @@ void UBehaviorTreeComponent::ExecuteTask(UBTTaskNode* TaskNode)
 
 		OnTaskFinished(TaskNode, TaskResult);
 	}
+
+	// It is now ok to resume any branch actions as the new active instance is set.
+	// Before that, decorators evaluating the IsExecutingBranch would be wrong.
+	ResumeBranchActions();
 }
 
 void UBehaviorTreeComponent::AbortCurrentTask()
