Fix CIS warning

diff --git a/Engine/Source/Developer/AITestSuite/Private/Tests/BTTest.cpp b/Engine/Source/Developer/AITestSuite/Private/Tests/BTTest.cpp
--- a/Engine/Source/Developer/AITestSuite/Private/Tests/BTTest.cpp
+++ b/Engine/Source/Developer/AITestSuite/Private/Tests/BTTest.cpp
@@ -2272,7 +2272,6 @@ struct FAITest_BTPostponeDuringSearch2 : public FAITest_SimpleBT
 	}
 };
 IMPLEMENT_AI_LATENT_TEST(FAITest_BTPostponeDuringSearch2, "System.AI.Behavior Trees.Service: Postpone during search (lower pri)")
-#undef LOCTEXT_NAMESPACE
 
 struct FAITest_BTStopDuringTaskExecute : public FAITest_SimpleBT
 {
