New BT Unittests that checks for
* StopTree with in all possible cases
* Unitialize with in all possible cases
* Cleanup with in all possible cases

#rnx

diff --git a/Engine/Source/Developer/AITestSuite/Classes/BehaviorTree/TestBTService_BTStopAction.h b/Engine/Source/Developer/AITestSuite/Classes/BehaviorTree/TestBTService_BTStopAction.h
new file mode 100644
--- /dev/null
+++ b/Engine/Source/Developer/AITestSuite/Classes/BehaviorTree/TestBTService_BTStopAction.h
@@ -0,0 +1,33 @@
+// Copyright Epic Games, Inc. All Rights Reserved.
+#pragma once
+
+#include "BehaviorTree/BTService.h"
+#include "TestBTStopAction.h"
+#include "TestBTService_BTStopAction.generated.h"
+
+UENUM()
+enum class EBTTestServiceStopTiming : uint8
+{
+	DuringBecomeRelevant,
+	DuringTick,
+	DuringCeaseRelevant,
+};
+
+UCLASS(meta = (HiddenNode))
+class UTestBTService_BTStopAction : public UBTService
+{
+	GENERATED_UCLASS_BODY()
+
+	UPROPERTY()
+	int32 LogIndex;
+
+	UPROPERTY()
+	EBTTestServiceStopTiming StopTiming;
+
+	UPROPERTY()
+	EBTTestStopAction StopAction;
+
+	virtual void OnBecomeRelevant(UBehaviorTreeComponent& OwnerComp, uint8* NodeMemory) override;
+	virtual void OnCeaseRelevant(UBehaviorTreeComponent& OwnerComp, uint8* NodeMemory) override;
+	virtual void TickNode(UBehaviorTreeComponent& OwnerComp, uint8* NodeMemory, float DeltaSeconds) override;
+};
diff --git a/Engine/Source/Developer/AITestSuite/Classes/BehaviorTree/TestBTService_StopTree.h b/Engine/Source/Developer/AITestSuite/Classes/BehaviorTree/TestBTService_StopTree.h
deleted file mode 100644
diff --git a/Engine/Source/Developer/AITestSuite/Classes/BehaviorTree/TestBTStopAction.h b/Engine/Source/Developer/AITestSuite/Classes/BehaviorTree/TestBTStopAction.h
new file mode 100644
--- /dev/null
+++ b/Engine/Source/Developer/AITestSuite/Classes/BehaviorTree/TestBTStopAction.h
@@ -0,0 +1,18 @@
+// Copyright Epic Games, Inc. All Rights Reserved.
+#pragma once
+
+#include "TestBTStopAction.generated.h"
+
+class UBehaviorTreeComponent;
+
+UENUM()
+enum class EBTTestStopAction : uint8
+{
+	StopTree,
+	UnInitialize,
+	Cleanup,
+	RestartTree,
+	StartTree,
+};
+
+void DoBTStopAction(UBehaviorTreeComponent& OwnerComp, const EBTTestStopAction StopAction);
\ No newline at end of file
diff --git a/Engine/Source/Developer/AITestSuite/Classes/BehaviorTree/TestBTTask_BTStopAction.h b/Engine/Source/Developer/AITestSuite/Classes/BehaviorTree/TestBTTask_BTStopAction.h
new file mode 100644
--- /dev/null
+++ b/Engine/Source/Developer/AITestSuite/Classes/BehaviorTree/TestBTTask_BTStopAction.h
@@ -0,0 +1,42 @@
+// Copyright Epic Games, Inc. All Rights Reserved.
+#pragma once
+
+#include "BehaviorTree/BTTaskNode.h"
+#include "TestBTStopAction.h"
+#include "TestBTTask_BTStopAction.generated.h"
+
+UENUM()
+enum class EBTTestTaskStopTiming : uint8
+{
+	DuringExecute,
+	DuringTick,
+	DuringAbort,
+	DuringFinish,
+};
+
+UCLASS(meta=(HiddenNode))
+class UTestBTTask_BTStopAction : public UBTTaskNode
+{
+	GENERATED_BODY()
+
+public:
+	UTestBTTask_BTStopAction(const FObjectInitializer& ObjectInitializer);
+
+	virtual EBTNodeResult::Type ExecuteTask(UBehaviorTreeComponent& OwnerComp, uint8* NodeMemory) override;
+	virtual EBTNodeResult::Type AbortTask(UBehaviorTreeComponent& OwnerComp, uint8* NodeMemory) override;
+
+	virtual void TickTask(class UBehaviorTreeComponent& OwnerComp, uint8* NodeMemory, float DeltaSeconds) override;
+	virtual void OnTaskFinished(UBehaviorTreeComponent& OwnerComp, uint8* NodeMemory, EBTNodeResult::Type TaskResult) override;
+
+	UPROPERTY()
+	EBTTestTaskStopTiming StopTiming;
+
+	UPROPERTY()
+	EBTTestStopAction StopAction;
+
+	UPROPERTY()
+	int32 LogIndex;
+
+	UPROPERTY()
+	TEnumAsByte<EBTNodeResult::Type> LogResult;
+};
diff --git a/Engine/Source/Developer/AITestSuite/Classes/BehaviorTree/TestBTTask_StopTree.h b/Engine/Source/Developer/AITestSuite/Classes/BehaviorTree/TestBTTask_StopTree.h
deleted file mode 100644
diff --git a/Engine/Source/Developer/AITestSuite/Private/BehaviorTree/TestBTService_BTStopAction.cpp b/Engine/Source/Developer/AITestSuite/Private/BehaviorTree/TestBTService_BTStopAction.cpp
new file mode 100644
--- /dev/null
+++ b/Engine/Source/Developer/AITestSuite/Private/BehaviorTree/TestBTService_BTStopAction.cpp
@@ -0,0 +1,55 @@
+// Copyright Epic Games, Inc. All Rights Reserved.
+
+#include "BehaviorTree/TestBTService_BTStopAction.h"
+#include "MockAI_BT.h"
+
+#include UE_INLINE_GENERATED_CPP_BY_NAME(TestBTService_BTStopAction)
+
+UTestBTService_BTStopAction::UTestBTService_BTStopAction(const FObjectInitializer& ObjectInitializer) : Super(ObjectInitializer)
+{
+	NodeName = "StopTreeService";
+
+	INIT_SERVICE_NODE_NOTIFY_FLAGS();
+
+	LogIndex = INDEX_NONE;
+	StopTiming = EBTTestServiceStopTiming::DuringTick;
+
+	// Force the service to tick every frame
+	Interval = 0.0f;
+	RandomDeviation = 0.0f;
+}
+
+void UTestBTService_BTStopAction::OnBecomeRelevant(UBehaviorTreeComponent& OwnerComp, uint8* NodeMemory)
+{
+	Super::OnBecomeRelevant(OwnerComp, NodeMemory);
+
+	if (StopTiming == EBTTestServiceStopTiming::DuringBecomeRelevant)
+	{
+		UMockAI_BT::ExecutionLog.Add(LogIndex);
+		DoBTStopAction(OwnerComp, StopAction);
+	}
+}
+
+void UTestBTService_BTStopAction::OnCeaseRelevant(UBehaviorTreeComponent& OwnerComp, uint8* NodeMemory)
+{
+	Super::OnCeaseRelevant(OwnerComp, NodeMemory);
+
+	if (StopTiming == EBTTestServiceStopTiming::DuringCeaseRelevant)
+	{
+		UMockAI_BT::ExecutionLog.Add(LogIndex);
+		OwnerComp.StopTree();
+		DoBTStopAction(OwnerComp, StopAction);
+	}
+}
+
+void UTestBTService_BTStopAction::TickNode(UBehaviorTreeComponent& OwnerComp, uint8* NodeMemory, float DeltaSeconds)
+{
+	Super::TickNode(OwnerComp, NodeMemory, DeltaSeconds);
+
+	if (StopTiming == EBTTestServiceStopTiming::DuringTick)
+	{
+		UMockAI_BT::ExecutionLog.Add(LogIndex);
+		OwnerComp.StopTree();
+		DoBTStopAction(OwnerComp, StopAction);
+	}
+}
\ No newline at end of file
diff --git a/Engine/Source/Developer/AITestSuite/Private/BehaviorTree/TestBTService_StopTree.cpp b/Engine/Source/Developer/AITestSuite/Private/BehaviorTree/TestBTService_StopTree.cpp
deleted file mode 100644
diff --git a/Engine/Source/Developer/AITestSuite/Private/BehaviorTree/TestBTTask_BTStopAction.cpp b/Engine/Source/Developer/AITestSuite/Private/BehaviorTree/TestBTTask_BTStopAction.cpp
new file mode 100644
--- /dev/null
+++ b/Engine/Source/Developer/AITestSuite/Private/BehaviorTree/TestBTTask_BTStopAction.cpp
@@ -0,0 +1,80 @@
+// Copyright Epic Games, Inc. All Rights Reserved.
+
+#include "BehaviorTree/TestBTTask_BTStopAction.h"
+#include "BehaviorTree/BehaviorTree.h"
+#include "MockAI_BT.h"
+
+#include UE_INLINE_GENERATED_CPP_BY_NAME(TestBTTask_BTStopAction)
+
+UTestBTTask_BTStopAction::UTestBTTask_BTStopAction(const FObjectInitializer& ObjectInitializer) : Super(ObjectInitializer)
+{
+	StopTiming = EBTTestTaskStopTiming::DuringExecute;
+	LogResult = EBTNodeResult::Succeeded;
+	LogIndex = INDEX_NONE;
+	INIT_TASK_NODE_NOTIFY_FLAGS();
+}
+
+EBTNodeResult::Type UTestBTTask_BTStopAction::ExecuteTask(UBehaviorTreeComponent& OwnerComp, uint8* NodeMemory)
+{
+	if (StopTiming == EBTTestTaskStopTiming::DuringExecute)
+	{
+		UMockAI_BT::ExecutionLog.Add(LogIndex);
+		DoBTStopAction(OwnerComp, StopAction);
+		return LogResult;
+	}
+
+	return EBTNodeResult::InProgress;
+}
+
+EBTNodeResult::Type UTestBTTask_BTStopAction::AbortTask(UBehaviorTreeComponent& OwnerComp, uint8* NodeMemory)
+{
+	if (StopTiming == EBTTestTaskStopTiming::DuringAbort)
+	{
+		UMockAI_BT::ExecutionLog.Add(LogIndex);
+		DoBTStopAction(OwnerComp, StopAction);
+	}
+	return EBTNodeResult::Aborted;
+}
+
+void UTestBTTask_BTStopAction::TickTask(class UBehaviorTreeComponent& OwnerComp, uint8* NodeMemory, float DeltaSeconds)
+{
+	if (StopTiming == EBTTestTaskStopTiming::DuringTick)
+	{
+		UMockAI_BT::ExecutionLog.Add(LogIndex);
+		DoBTStopAction(OwnerComp, StopAction);
+		FinishLatentTask(OwnerComp, LogResult);
+	}
+}
+
+void UTestBTTask_BTStopAction::OnTaskFinished(UBehaviorTreeComponent& OwnerComp, uint8* NodeMemory, EBTNodeResult::Type TaskResult)
+{
+	if (StopTiming == EBTTestTaskStopTiming::DuringFinish)
+	{
+		UMockAI_BT::ExecutionLog.Add(LogIndex);
+		DoBTStopAction(OwnerComp, StopAction);
+	}
+}
+
+
+void DoBTStopAction(UBehaviorTreeComponent& OwnerComp, const EBTTestStopAction StopAction)
+{
+	switch (StopAction)
+	{
+	case EBTTestStopAction::StopTree:
+		OwnerComp.StopTree();
+		break;
+	case EBTTestStopAction::UnInitialize:
+		OwnerComp.UninitializeComponent();
+		break;
+	case EBTTestStopAction::Cleanup:
+		OwnerComp.Cleanup();
+		break;
+	case EBTTestStopAction::RestartTree:
+		OwnerComp.RestartTree();
+		break;
+	case EBTTestStopAction::StartTree:
+		UBehaviorTree* NewBTAsset = NewObject<UBehaviorTree>(GetTransientPackage(), NAME_None, RF_NoFlags, OwnerComp.GetRootTree());
+		OwnerComp.StartTree(*NewBTAsset);
+		break;
+	}
+}
\ No newline at end of file
diff --git a/Engine/Source/Developer/AITestSuite/Private/BehaviorTree/TestBTTask_StopTree.cpp b/Engine/Source/Developer/AITestSuite/Private/BehaviorTree/TestBTTask_StopTree.cpp
deleted file mode 100644
diff --git a/Engine/Source/Developer/AITestSuite/Private/MockAI/MockAI_BT.cpp b/Engine/Source/Developer/AITestSuite/Private/MockAI/MockAI_BT.cpp
--- a/Engine/Source/Developer/AITestSuite/Private/MockAI/MockAI_BT.cpp
+++ b/Engine/Source/Developer/AITestSuite/Private/MockAI/MockAI_BT.cpp
@@ -21,6 +21,8 @@ UMockAI_BT::UMockAI_BT(const FObjectInitializer& ObjectInitializer)
 		UseBrainComponent<UBehaviorTreeComponent>();
 
 		BTComp = Cast<UBehaviorTreeComponent>(BrainComp);
+		BTComp->RegisterComponent();
+		BTComp->InitializeComponent();
 	}
 }
 
diff --git a/Engine/Source/Developer/AITestSuite/Private/Tests/BTTest.cpp b/Engine/Source/Developer/AITestSuite/Private/Tests/BTTest.cpp
--- a/Engine/Source/Developer/AITestSuite/Private/Tests/BTTest.cpp
+++ b/Engine/Source/Developer/AITestSuite/Private/Tests/BTTest.cpp
@@ -2149,14 +2149,14 @@ struct FAITest_BTSubtreeAbortOutToLowerPrio2 : public FAITest_SimpleBT
 			FBTBuilder::AddTask(CompNode, MainTaskExecute, EBTNodeResult::Succeeded, 1, MainTaskTick);
 		}
 
-		ExpectedResult.Add(SubtreeTaskExecute);
-		ExpectedResult.Add(SubtreeTaskTick);
-		ExpectedResult.Add(SubtreeTaskTick);
-		ExpectedResult.Add(BBDecoratorBecomeRelevant);
-		ExpectedResult.Add(MainTaskExecute);
-		ExpectedResult.Add(MainTaskTick);
-		ExpectedResult.Add(MainTaskTick);
-		ExpectedResult.Add(BBDecoratorCeaseRelevant);
+		ExpectedResult.Add(SubtreeTaskExecute/*1*/);
+		ExpectedResult.Add(SubtreeTaskTick/*2*/);
+		ExpectedResult.Add(SubtreeTaskTick/*2*/);
+		ExpectedResult.Add(BBDecoratorBecomeRelevant/*3*/);
+		ExpectedResult.Add(MainTaskExecute/*5*/);
+		ExpectedResult.Add(MainTaskTick/*6*/);
+		ExpectedResult.Add(MainTaskTick/*6*/);
+		ExpectedResult.Add(BBDecoratorCeaseRelevant/*4*/);
 	}
 };
 IMPLEMENT_AI_LATENT_TEST(FAITest_BTSubtreeAbortOutToLowerPrio2, "System.AI.Behavior Trees.Subtree: abort out lower prio 2")
@@ -2325,25 +2325,38 @@ struct FAITest_BTPostponeDuringSearch2 : public FAITest_SimpleBT
 };
 IMPLEMENT_AI_LATENT_TEST(FAITest_BTPostponeDuringSearch2, "System.AI.Behavior Trees.Service: Postpone during search (lower pri)")
 
+//////////////////////////////////////////////////////////////////////////////////////////////////
+// Stop tree tests
+//////////////////////////////////////////////////////////////////////////////////////////////////
 struct FAITest_BTStopDuringTaskExecute : public FAITest_SimpleBT
 {
 	FAITest_BTStopDuringTaskExecute()
 	{
+		enum
+		{
+			ServiceActivate = 1,
+			ServiceDeactivate,
+			ServiceTick,
+			UninitTaskExecute,
+		};
+
 		UBTCompositeNode& RootNode = FBTBuilder::AddSelector(*BTAsset);
 		UBTCompositeNode& CompNode = FBTBuilder::AddSelector(RootNode);
 		{
-			FBTBuilder::WithServiceLog(CompNode, 1/*ActivationIndex*/, 2/*DeactivationIndex*/, 3/*TickIndex*/);
+			// Service
+			FBTBuilder::WithServiceLog(CompNode, ServiceActivate, ServiceDeactivate, ServiceTick);
 
 			UBTCompositeNode& Comp1Node = FBTBuilder::AddSelector(CompNode);
 			{
-				FBTBuilder::AddTaskStopTree(Comp1Node, 4, EBTNodeResult::Succeeded, EBTTestTaskStopTree::DuringExecute);
+				// UninitTask
+				FBTBuilder::AddTaskBTStopAction(Comp1Node, UninitTaskExecute, EBTNodeResult::Succeeded, EBTTestTaskStopTiming::DuringExecute, EBTTestStopAction::StopTree);
 			}
 		}
 
-		ExpectedResult.Add(1);
-		ExpectedResult.Add(3);
-		ExpectedResult.Add(4);
-		ExpectedResult.Add(2);
+		ExpectedResult.Add(ServiceActivate/*1*/);
+		ExpectedResult.Add(ServiceTick/*3*/);
+		ExpectedResult.Add(UninitTaskExecute/*4*/);
+		ExpectedResult.Add(ServiceDeactivate/*2*/);
 	}
 };
 IMPLEMENT_AI_LATENT_TEST(FAITest_BTStopDuringTaskExecute, "System.AI.Behavior Trees.Stop: during task execute")
@@ -2352,21 +2365,31 @@ struct FAITest_BTStopDuringTaskTick : public FAITest_SimpleBT
 {
 	FAITest_BTStopDuringTaskTick()
 	{
+		enum
+		{
+			ServiceActivate = 1,
+			ServiceDeactivate,
+			ServiceTick,
+			UninitTaskTick,
+		};
+
 		UBTCompositeNode& RootNode = FBTBuilder::AddSelector(*BTAsset);
 		UBTCompositeNode& CompNode = FBTBuilder::AddSelector(RootNode);
 		{
-			FBTBuilder::WithServiceLog(CompNode, 1/*ActivationIndex*/, 2/*DeactivationIndex*/, 3/*TickIndex*/);
+			// Service
+			FBTBuilder::WithServiceLog(CompNode, ServiceActivate, ServiceDeactivate, ServiceTick);
 
 			UBTCompositeNode& Comp1Node = FBTBuilder::AddSelector(CompNode);
 			{
-				FBTBuilder::AddTaskStopTree(Comp1Node, 4, EBTNodeResult::Succeeded, EBTTestTaskStopTree::DuringTick);
+				// UninitTask
+				FBTBuilder::AddTaskBTStopAction(Comp1Node, UninitTaskTick, EBTNodeResult::Succeeded, EBTTestTaskStopTiming::DuringTick, EBTTestStopAction::StopTree);
 			}
 		}
 
-		ExpectedResult.Add(1);
-		ExpectedResult.Add(3);
-		ExpectedResult.Add(4);
-		ExpectedResult.Add(2);
+		ExpectedResult.Add(ServiceActivate/*1*/);
+		ExpectedResult.Add(ServiceTick/*3*/);
+		ExpectedResult.Add(UninitTaskTick/*4*/);
+		ExpectedResult.Add(ServiceDeactivate/*2*/);
 	}
 };
 IMPLEMENT_AI_LATENT_TEST(FAITest_BTStopDuringTaskTick, "System.AI.Behavior Trees.Stop: during task tick")
@@ -2376,31 +2399,47 @@ struct FAITest_BTStopDuringTaskAbort : public FAITest_SimpleBT
 {
 	FAITest_BTStopDuringTaskAbort()
 	{
+		enum
+		{
+			ServiceActivate = 1,
+			ServiceDeactivate,
+			ServiceTick,
+			SubServiceActivate,
+			SubServiceDeactivate,
+			SubServiceTick,
+			UninitTaskAborting,
+			TaskExecute,
+		};
+
 		UBTCompositeNode& RootNode = FBTBuilder::AddSelector(*BTAsset);
 		UBTCompositeNode& CompNode = FBTBuilder::AddSelector(RootNode);
 		{
-			FBTBuilder::WithServiceLog(CompNode, 1/*ActivationIndex*/, 2/*DeactivationIndex*/, 3/*TickIndex*/);
+			// Service
+			FBTBuilder::WithServiceLog(CompNode, ServiceActivate, ServiceDeactivate, ServiceTick);
 
 			UBTCompositeNode& Comp1Node = FBTBuilder::AddSelector(CompNode);
 			{
-				FBTBuilder::WithServiceLog(Comp1Node, 4/*ActivationIndex*/, 5/*DeactivationIndex*/, 6/*TickIndex*/, TEXT("Bool1"));
+				// SubService
+				FBTBuilder::WithServiceLog(Comp1Node, SubServiceActivate, SubServiceDeactivate, SubServiceTick, TEXT("Bool1"));
 
-				FBTBuilder::AddTaskStopTree(Comp1Node, 7, EBTNodeResult::Succeeded, EBTTestTaskStopTree::DuringAbort);
+				// UninitTask
+				FBTBuilder::AddTaskBTStopAction(Comp1Node, UninitTaskAborting, EBTNodeResult::Succeeded, EBTTestTaskStopTiming::DuringAbort, EBTTestStopAction::StopTree);
 				FBTBuilder::WithDecoratorBlackboard(Comp1Node, EBasicKeyOperation::NotSet, EBTFlowAbortMode::Both, TEXT("Bool1"));
 
-				FBTBuilder::AddTask(Comp1Node, 8, EBTNodeResult::Succeeded);
+				// Task
+				FBTBuilder::AddTask(Comp1Node, TaskExecute, EBTNodeResult::Succeeded);
 			}
 		}
 
-		ExpectedResult.Add(1);
-		ExpectedResult.Add(4);
-		ExpectedResult.Add(3);
-		ExpectedResult.Add(6);
-		ExpectedResult.Add(3);
-		ExpectedResult.Add(6);
-		ExpectedResult.Add(7);
-		ExpectedResult.Add(2);
-		ExpectedResult.Add(5);
+		ExpectedResult.Add(ServiceActivate/*1*/);
+		ExpectedResult.Add(SubServiceActivate/*4*/);
+		ExpectedResult.Add(ServiceTick/*5*/);
+		ExpectedResult.Add(SubServiceTick/*6*/);
+		ExpectedResult.Add(ServiceTick/*3*/);
+		ExpectedResult.Add(SubServiceTick/*6*/);
+		ExpectedResult.Add(UninitTaskAborting/*7*/);
+		ExpectedResult.Add(ServiceDeactivate/*2*/);
+		ExpectedResult.Add(SubServiceDeactivate/*5*/);
 	}
 };
 IMPLEMENT_AI_LATENT_TEST(FAITest_BTStopDuringTaskAbort, "System.AI.Behavior Trees.Stop: during task abort")
@@ -2410,31 +2449,47 @@ struct FAITest_BTStopDuringTaskFinished : public FAITest_SimpleBT
 {
 	FAITest_BTStopDuringTaskFinished()
 	{
+		enum
+		{
+			ServiceActivate = 1,
+			ServiceDeactivate,
+			ServiceTick,
+			SubServiceActivate,
+			SubServiceDeactivate,
+			SubServiceTick,
+			UninitTaskFinished,
+			TaskExecute,
+		};
+
 		UBTCompositeNode& RootNode = FBTBuilder::AddSelector(*BTAsset);
 		UBTCompositeNode& CompNode = FBTBuilder::AddSelector(RootNode);
 		{
-			FBTBuilder::WithServiceLog(CompNode, 1/*ActivationIndex*/, 2/*DeactivationIndex*/, 3/*TickIndex*/);
+			// Service
+			FBTBuilder::WithServiceLog(CompNode, ServiceActivate, ServiceDeactivate, ServiceTick);
 
 			UBTCompositeNode& Comp1Node = FBTBuilder::AddSelector(CompNode);
 			{
-				FBTBuilder::WithServiceLog(Comp1Node, 4/*ActivationIndex*/, 5/*DeactivationIndex*/, 6/*TickIndex*/, TEXT("Bool1"));
+				// SubService
+				FBTBuilder::WithServiceLog(Comp1Node, SubServiceActivate, SubServiceDeactivate, SubServiceTick, TEXT("Bool1"));
 
-				FBTBuilder::AddTaskStopTree(Comp1Node, 7, EBTNodeResult::Succeeded, EBTTestTaskStopTree::DuringFinish);
+				// UninitTask
+				FBTBuilder::AddTaskBTStopAction(Comp1Node, UninitTaskFinished, EBTNodeResult::Succeeded, EBTTestTaskStopTiming::DuringFinish, EBTTestStopAction::StopTree);
 				FBTBuilder::WithDecoratorBlackboard(Comp1Node, EBasicKeyOperation::NotSet, EBTFlowAbortMode::Both, TEXT("Bool1"));
 
-				FBTBuilder::AddTask(Comp1Node, 8, EBTNodeResult::Succeeded);
+				// Task
+				FBTBuilder::AddTask(Comp1Node, TaskExecute, EBTNodeResult::Succeeded);
 			}
 		}
 
-		ExpectedResult.Add(1);
-		ExpectedResult.Add(4);
-		ExpectedResult.Add(3);
-		ExpectedResult.Add(6);
-		ExpectedResult.Add(3);
-		ExpectedResult.Add(6);
-		ExpectedResult.Add(7);
-		ExpectedResult.Add(2);
-		ExpectedResult.Add(5);
+		ExpectedResult.Add(ServiceActivate/*1*/);
+		ExpectedResult.Add(SubServiceActivate/*4*/);
+		ExpectedResult.Add(ServiceTick/*3*/);
+		ExpectedResult.Add(SubServiceTick/*6*/);
+		ExpectedResult.Add(ServiceTick/*3*/);
+		ExpectedResult.Add(SubServiceTick/*6*/);
+		ExpectedResult.Add(UninitTaskFinished/*7*/);
+		ExpectedResult.Add(ServiceDeactivate/*2*/);
+		ExpectedResult.Add(SubServiceDeactivate/*5*/);
 	}
 };
 IMPLEMENT_AI_LATENT_TEST(FAITest_BTStopDuringTaskFinished, "System.AI.Behavior Trees.Stop: during task finished")
@@ -2443,31 +2498,49 @@ struct FAITest_BTStopDuringServiceTick : public FAITest_SimpleBT
 {
 	FAITest_BTStopDuringServiceTick()
 	{
+		enum
+		{
+			ServiceActivate = 1,
+			ServiceDeactivate,
+			ServiceTick,
+			SubServiceActivate,
+			SubServiceDeactivate,
+			SubServiceTick,
+			UninitServiceTick,
+			TaskExecute,
+			Task2Execute,
+		};
+
 		UBTCompositeNode& RootNode = FBTBuilder::AddSelector(*BTAsset);
 		UBTCompositeNode& CompNode = FBTBuilder::AddSelector(RootNode);
 		{
-			FBTBuilder::WithServiceLog(CompNode, 1/*ActivationIndex*/, 2/*DeactivationIndex*/, 3/*TickIndex*/);
+			// Servier
+			FBTBuilder::WithServiceLog(CompNode, ServiceActivate, ServiceDeactivate, ServiceTick);
 
 			UBTCompositeNode& Comp1Node = FBTBuilder::AddSelector(CompNode);
 			{
-				FBTBuilder::WithServiceLog(Comp1Node, 4/*ActivationIndex*/, 5/*DeactivationIndex*/, 6/*TickIndex*/, TEXT("Bool1"));
-				FBTBuilder::WithServiceStopTree(Comp1Node, 9, EBTTestServiceStopTree::DuringTick);
+				// SubService
+				FBTBuilder::WithServiceLog(Comp1Node, SubServiceActivate, SubServiceDeactivate, SubServiceTick, TEXT("Bool1"));
+				// UninitServiceTick
+				FBTBuilder::WithServiceBTStopAction(Comp1Node, UninitServiceTick, EBTTestServiceStopTiming::DuringTick, EBTTestStopAction::StopTree);
 
-				FBTBuilder::AddTask(Comp1Node, 7, EBTNodeResult::Succeeded);
+				// Task
+				FBTBuilder::AddTask(Comp1Node, TaskExecute, EBTNodeResult::Succeeded);
 				FBTBuilder::WithDecoratorBlackboard(Comp1Node, EBasicKeyOperation::NotSet, EBTFlowAbortMode::Both, TEXT("Bool1"));
 
-				FBTBuilder::AddTask(Comp1Node, 8, EBTNodeResult::Succeeded);
+				// Task2
+				FBTBuilder::AddTask(Comp1Node, Task2Execute, EBTNodeResult::Succeeded);
 			}
 		}
 
-		ExpectedResult.Add(1);
-		ExpectedResult.Add(4);
-		ExpectedResult.Add(3);
-		ExpectedResult.Add(6);
-		ExpectedResult.Add(9);
-		ExpectedResult.Add(7);
-		ExpectedResult.Add(2);
-		ExpectedResult.Add(5);
+		ExpectedResult.Add(ServiceActivate/*1*/);
+		ExpectedResult.Add(SubServiceActivate/*4*/);
+		ExpectedResult.Add(ServiceTick/*3*/);
+		ExpectedResult.Add(SubServiceTick/*6*/);
+		ExpectedResult.Add(UninitServiceTick/*7*/);
+		ExpectedResult.Add(TaskExecute/*8*/);
+		ExpectedResult.Add(ServiceDeactivate/*2*/);
+		ExpectedResult.Add(SubServiceDeactivate/*5*/);
 	}
 };
 IMPLEMENT_AI_LATENT_TEST(FAITest_BTStopDuringServiceTick, "System.AI.Behavior Trees.Stop: during service tick")
@@ -2476,31 +2549,49 @@ struct FAITest_BTStopDuringServiceBecomeRelevant : public FAITest_SimpleBT
 {
 	FAITest_BTStopDuringServiceBecomeRelevant()
 	{
+		enum
+		{
+			ServiceActivate = 1,
+			ServiceDeactivate,
+			ServiceTick,
+			SubServiceActivate,
+			SubServiceDeactivate,
+			SubServiceTick,
+			UninitServiceBecomeRelevant,
+			TaskExecute,
+			Task2Execute,
+		};
+
 		UBTCompositeNode& RootNode = FBTBuilder::AddSelector(*BTAsset);
 		UBTCompositeNode& CompNode = FBTBuilder::AddSelector(RootNode);
 		{
-			FBTBuilder::WithServiceLog(CompNode, 1/*ActivationIndex*/, 2/*DeactivationIndex*/, 3/*TickIndex*/);
+			// Service
+			FBTBuilder::WithServiceLog(CompNode, ServiceActivate, ServiceDeactivate, ServiceTick);
 
 			UBTCompositeNode& Comp1Node = FBTBuilder::AddSelector(CompNode);
 			{
-				FBTBuilder::WithServiceLog(Comp1Node, 4/*ActivationIndex*/, 5/*DeactivationIndex*/, 6/*TickIndex*/, TEXT("Bool1"));
-				FBTBuilder::WithServiceStopTree(Comp1Node, 9, EBTTestServiceStopTree::DuringBecomeRelevant);
+				// SubService
+				FBTBuilder::WithServiceLog(Comp1Node, SubServiceActivate, SubServiceDeactivate, SubServiceTick, TEXT("Bool1"));
+				// UninitService
+				FBTBuilder::WithServiceBTStopAction(Comp1Node, UninitServiceBecomeRelevant, EBTTestServiceStopTiming::DuringBecomeRelevant, EBTTestStopAction::StopTree);
 
-				FBTBuilder::AddTask(Comp1Node, 7, EBTNodeResult::Succeeded);
+				// Task
+				FBTBuilder::AddTask(Comp1Node, TaskExecute, EBTNodeResult::Succeeded);
 				FBTBuilder::WithDecoratorBlackboard(Comp1Node, EBasicKeyOperation::NotSet, EBTFlowAbortMode::Both, TEXT("Bool1"));
 
-				FBTBuilder::AddTask(Comp1Node, 8, EBTNodeResult::Succeeded);
+				// Task2
+				FBTBuilder::AddTask(Comp1Node, Task2Execute, EBTNodeResult::Succeeded);
 			}
 		}
 
-		ExpectedResult.Add(1);
-		ExpectedResult.Add(4);
-		ExpectedResult.Add(9);
-		ExpectedResult.Add(3);
-		ExpectedResult.Add(6);
-		ExpectedResult.Add(7);
-		ExpectedResult.Add(2);
-		ExpectedResult.Add(5);
+		ExpectedResult.Add(ServiceActivate/*1*/);
+		ExpectedResult.Add(SubServiceActivate/*4*/);
+		ExpectedResult.Add(UninitServiceBecomeRelevant/*7*/);
+		ExpectedResult.Add(ServiceTick/*3*/);
+		ExpectedResult.Add(SubServiceTick/*6*/);
+		ExpectedResult.Add(TaskExecute/*8*/);
+		ExpectedResult.Add(ServiceDeactivate/*2*/);
+		ExpectedResult.Add(SubServiceDeactivate/*5*/);
 	}
 };
 IMPLEMENT_AI_LATENT_TEST(FAITest_BTStopDuringServiceBecomeRelevant, "System.AI.Behavior Trees.Stop: during service become relevant")
@@ -2509,33 +2600,51 @@ struct FAITest_BTStopDuringServiceCeaseRelevant : public FAITest_SimpleBT
 {
 	FAITest_BTStopDuringServiceCeaseRelevant()
 	{
+		enum
+		{
+			ServiceActivate = 1,
+			ServiceDeactivate,
+			ServiceTick,
+			SubServiceActivate,
+			SubServiceDeactivate,
+			SubServiceTick,
+			UninitServiceCeaseRelevant,
+			TaskExecute,
+			Task2Execute,
+		};
+
 		UBTCompositeNode& RootNode = FBTBuilder::AddSelector(*BTAsset);
 		UBTCompositeNode& CompNode = FBTBuilder::AddSelector(RootNode);
 		{
-			FBTBuilder::WithServiceLog(CompNode, 1/*ActivationIndex*/, 2/*DeactivationIndex*/, 3/*TickIndex*/);
+			// Service
+			FBTBuilder::WithServiceLog(CompNode, ServiceActivate, ServiceDeactivate, ServiceTick);
 
 			UBTCompositeNode& Comp1Node = FBTBuilder::AddSelector(CompNode);
 			{
-				FBTBuilder::WithServiceLog(Comp1Node, 4/*ActivationIndex*/, 5/*DeactivationIndex*/, 6/*TickIndex*/, TEXT("Bool1"));
-				FBTBuilder::WithServiceStopTree(Comp1Node, 9, EBTTestServiceStopTree::DuringCeaseRelevant);
+				// SubService
+				FBTBuilder::WithServiceLog(Comp1Node, SubServiceActivate, SubServiceDeactivate, SubServiceTick, TEXT("Bool1"));
+				// UninitService
+				FBTBuilder::WithServiceBTStopAction(Comp1Node, UninitServiceCeaseRelevant, EBTTestServiceStopTiming::DuringCeaseRelevant, EBTTestStopAction::StopTree);
 
-				FBTBuilder::AddTask(Comp1Node, 7, EBTNodeResult::Succeeded);
+				// Task
+				FBTBuilder::AddTask(Comp1Node, TaskExecute, EBTNodeResult::Succeeded);
 				FBTBuilder::WithDecoratorBlackboard(Comp1Node, EBasicKeyOperation::NotSet, EBTFlowAbortMode::Both, TEXT("Bool1"));
 			}
-			FBTBuilder::AddTask(CompNode, 8, EBTNodeResult::Succeeded);
+			// Task2
+			FBTBuilder::AddTask(CompNode, Task2Execute, EBTNodeResult::Succeeded);
 		}
 
-		ExpectedResult.Add(1);
-		ExpectedResult.Add(4);
-		ExpectedResult.Add(3);
-		ExpectedResult.Add(6);
-		ExpectedResult.Add(7);
-		ExpectedResult.Add(3);
-		ExpectedResult.Add(6);
-		ExpectedResult.Add(5);
-		ExpectedResult.Add(9);
-		ExpectedResult.Add(8);
-		ExpectedResult.Add(2);
+		ExpectedResult.Add(ServiceActivate/*1*/);
+		ExpectedResult.Add(SubServiceActivate/*4*/);
+		ExpectedResult.Add(ServiceTick/*3*/);
+		ExpectedResult.Add(SubServiceTick/*6*/);
+		ExpectedResult.Add(TaskExecute/*8*/);
+		ExpectedResult.Add(ServiceTick/*3*/);
+		ExpectedResult.Add(SubServiceTick/*6*/);
+		ExpectedResult.Add(SubServiceDeactivate/*5*/);
+		ExpectedResult.Add(UninitServiceCeaseRelevant/*7*/);
+		ExpectedResult.Add(Task2Execute/*9*/);
+		ExpectedResult.Add(ServiceDeactivate/*2*/);
 	}
 };
 IMPLEMENT_AI_LATENT_TEST(FAITest_BTStopDuringServiceCeaseRelevant, "System.AI.Behavior Trees.Stop: during service cease relevant")
@@ -2544,31 +2653,44 @@ struct FAITest_BTStopDuringTaskServiceTick : public FAITest_SimpleBT
 {
 	FAITest_BTStopDuringTaskServiceTick()
 	{
+		enum
+		{
+			ServiceActivate = 1,
+			ServiceDeactivate,
+			ServiceTick,
+			SubServiceActivate,
+			SubServiceDeactivate,
+			SubServiceTick,
+			UninitTaskServiceTick,
+			TaskExecute,
+			Task2Execute,
+		};
+
 		UBTCompositeNode& RootNode = FBTBuilder::AddSelector(*BTAsset);
 		UBTCompositeNode& CompNode = FBTBuilder::AddSelector(RootNode);
 		{
-			FBTBuilder::WithServiceLog(CompNode, 1/*ActivationIndex*/, 2/*DeactivationIndex*/, 3/*TickIndex*/);
+			FBTBuilder::WithServiceLog(CompNode, ServiceActivate, ServiceDeactivate, ServiceTick);
 
 			UBTCompositeNode& Comp1Node = FBTBuilder::AddSelector(CompNode);
 			{
-				FBTBuilder::WithServiceLog(Comp1Node, 4/*ActivationIndex*/, 5/*DeactivationIndex*/, 6/*TickIndex*/, TEXT("Bool1"));
+				FBTBuilder::WithServiceLog(Comp1Node, SubServiceActivate, SubServiceDeactivate, SubServiceTick, TEXT("Bool1"));
 
-				FBTBuilder::AddTask(Comp1Node, 7, EBTNodeResult::Succeeded);
+				FBTBuilder::AddTask(Comp1Node, TaskExecute, EBTNodeResult::Succeeded);
 				FBTBuilder::WithDecoratorBlackboard(Comp1Node, EBasicKeyOperation::NotSet, EBTFlowAbortMode::Both, TEXT("Bool1"));
-				FBTBuilder::WithTaskServiceStopTree(Comp1Node, 9, EBTTestServiceStopTree::DuringTick);
+				FBTBuilder::WithTaskServiceBTStopAction(Comp1Node, UninitTaskServiceTick, EBTTestServiceStopTiming::DuringTick, EBTTestStopAction::StopTree);
 
-				FBTBuilder::AddTask(Comp1Node, 8, EBTNodeResult::Succeeded);
+				FBTBuilder::AddTask(Comp1Node, Task2Execute, EBTNodeResult::Succeeded);
 			}
 		}
 
-		ExpectedResult.Add(1);
-		ExpectedResult.Add(4);
-		ExpectedResult.Add(3);
-		ExpectedResult.Add(6);
-		ExpectedResult.Add(9);
-		ExpectedResult.Add(7);
-		ExpectedResult.Add(2);
-		ExpectedResult.Add(5);
+		ExpectedResult.Add(ServiceActivate/*1*/);
+		ExpectedResult.Add(SubServiceActivate/*4*/);
+		ExpectedResult.Add(ServiceTick/*3*/);
+		ExpectedResult.Add(SubServiceTick/*6*/);
+		ExpectedResult.Add(UninitTaskServiceTick/*7*/);
+		ExpectedResult.Add(TaskExecute/*8*/);
+		ExpectedResult.Add(ServiceDeactivate/*2*/);
+		ExpectedResult.Add(SubServiceDeactivate/*5*/);
 	}
 };
 IMPLEMENT_AI_LATENT_TEST(FAITest_BTStopDuringTaskServiceTick, "System.AI.Behavior Trees.Stop: during task service tick")
@@ -2577,31 +2699,44 @@ struct FAITest_BTStopDuringTaskServiceBecomeRelevant : public FAITest_SimpleBT
 {
 	FAITest_BTStopDuringTaskServiceBecomeRelevant()
 	{
+		enum
+		{
+			ServiceActivate = 1,
+			ServiceDeactivate,
+			ServiceTick,
+			SubServiceActivate,
+			SubServiceDeactivate,
+			SubServiceTick,
+			UninitTaskServiceBecomeRelevant,
+			TaskExecute,
+			Task2Execute,
+		};
+
 		UBTCompositeNode& RootNode = FBTBuilder::AddSelector(*BTAsset);
 		UBTCompositeNode& CompNode = FBTBuilder::AddSelector(RootNode);
 		{
-			FBTBuilder::WithServiceLog(CompNode, 1/*ActivationIndex*/, 2/*DeactivationIndex*/, 3/*TickIndex*/);
+			FBTBuilder::WithServiceLog(CompNode, ServiceActivate, ServiceDeactivate, ServiceTick);
 
 			UBTCompositeNode& Comp1Node = FBTBuilder::AddSelector(CompNode);
 			{
-				FBTBuilder::WithServiceLog(Comp1Node, 4/*ActivationIndex*/, 5/*DeactivationIndex*/, 6/*TickIndex*/, TEXT("Bool1"));
+				FBTBuilder::WithServiceLog(Comp1Node, SubServiceActivate, SubServiceDeactivate, SubServiceTick, TEXT("Bool1"));
 
-				FBTBuilder::AddTask(Comp1Node, 7, EBTNodeResult::Succeeded);
+				FBTBuilder::AddTask(Comp1Node, TaskExecute, EBTNodeResult::Succeeded);
 				FBTBuilder::WithDecoratorBlackboard(Comp1Node, EBasicKeyOperation::NotSet, EBTFlowAbortMode::Both, TEXT("Bool1"));
-				FBTBuilder::WithTaskServiceStopTree(Comp1Node, 9, EBTTestServiceStopTree::DuringBecomeRelevant);
+				FBTBuilder::WithTaskServiceBTStopAction(Comp1Node, UninitTaskServiceBecomeRelevant, EBTTestServiceStopTiming::DuringBecomeRelevant, EBTTestStopAction::StopTree);
 
-				FBTBuilder::AddTask(Comp1Node, 8, EBTNodeResult::Succeeded);
+				FBTBuilder::AddTask(Comp1Node, Task2Execute, EBTNodeResult::Succeeded);
 			}
 		}
 
-		ExpectedResult.Add(1);
-		ExpectedResult.Add(4);
-		ExpectedResult.Add(3);
-		ExpectedResult.Add(6);
-		ExpectedResult.Add(9);
-		ExpectedResult.Add(7);
-		ExpectedResult.Add(2);
-		ExpectedResult.Add(5);
+		ExpectedResult.Add(ServiceActivate/*1*/);
+		ExpectedResult.Add(SubServiceActivate/*4*/);
+		ExpectedResult.Add(ServiceTick/*3*/);
+		ExpectedResult.Add(SubServiceTick/*6*/);
+		ExpectedResult.Add(UninitTaskServiceBecomeRelevant/*7*/);
+		ExpectedResult.Add(TaskExecute/*8*/);
+		ExpectedResult.Add(ServiceDeactivate/*2*/);
+		ExpectedResult.Add(SubServiceDeactivate/*5*/);
 	}
 };
 IMPLEMENT_AI_LATENT_TEST(FAITest_BTStopDuringTaskServiceBecomeRelevant, "System.AI.Behavior Trees.Stop: during task service become relevant")
@@ -2611,36 +2746,982 @@ struct FAITest_BTStopDuringTaskServiceCeaseRelevant : public FAITest_SimpleBT
 {
 	FAITest_BTStopDuringTaskServiceCeaseRelevant()
 	{
+		enum
+		{
+			ServiceActivate = 1,
+			ServiceDeactivate,
+			ServiceTick,
+			SubServiceActivate,
+			SubServiceDeactivate,
+			SubServiceTick,
+			UninitTaskServiceCeaseRelevant,
+			TaskExecute,
+			Task2Execute,
+		};
+
 		UBTCompositeNode& RootNode = FBTBuilder::AddSelector(*BTAsset);
 		UBTCompositeNode& CompNode = FBTBuilder::AddSelector(RootNode);
 		{
-			FBTBuilder::WithServiceLog(CompNode, 1/*ActivationIndex*/, 2/*DeactivationIndex*/, 3/*TickIndex*/);
+			FBTBuilder::WithServiceLog(CompNode, ServiceActivate, ServiceDeactivate, ServiceTick);
 
 			UBTCompositeNode& Comp1Node = FBTBuilder::AddSelector(CompNode);
 			{
-				FBTBuilder::WithServiceLog(Comp1Node, 4/*ActivationIndex*/, 5/*DeactivationIndex*/, 6/*TickIndex*/, TEXT("Bool1"));
+				FBTBuilder::WithServiceLog(Comp1Node, SubServiceActivate, SubServiceDeactivate, SubServiceTick, TEXT("Bool1"));
 
-				FBTBuilder::AddTask(Comp1Node, 7, EBTNodeResult::Succeeded);
+				FBTBuilder::AddTask(Comp1Node, TaskExecute, EBTNodeResult::Succeeded);
 				FBTBuilder::WithDecoratorBlackboard(Comp1Node, EBasicKeyOperation::NotSet, EBTFlowAbortMode::Both, TEXT("Bool1"));
-				FBTBuilder::WithTaskServiceStopTree(Comp1Node, 9, EBTTestServiceStopTree::DuringCeaseRelevant);
+				FBTBuilder::WithTaskServiceBTStopAction(Comp1Node, UninitTaskServiceCeaseRelevant, EBTTestServiceStopTiming::DuringCeaseRelevant, EBTTestStopAction::StopTree);
 
-				FBTBuilder::AddTask(CompNode, 8, EBTNodeResult::Succeeded);
+				FBTBuilder::AddTask(CompNode, Task2Execute, EBTNodeResult::Succeeded);
 			}
 		}
 
-		ExpectedResult.Add(1);
-		ExpectedResult.Add(4);
-		ExpectedResult.Add(3);
-		ExpectedResult.Add(6);
-		ExpectedResult.Add(7);
-		ExpectedResult.Add(3);
-		ExpectedResult.Add(6);
-		ExpectedResult.Add(9);
-		ExpectedResult.Add(5);
-		ExpectedResult.Add(8);
-		ExpectedResult.Add(2);
+		ExpectedResult.Add(ServiceActivate/*1*/);
+		ExpectedResult.Add(SubServiceActivate/*4*/);
+		ExpectedResult.Add(ServiceTick/*3*/);
+		ExpectedResult.Add(SubServiceTick/*6*/);
+		ExpectedResult.Add(TaskExecute/*8*/);
+		ExpectedResult.Add(ServiceTick/*3*/);
+		ExpectedResult.Add(SubServiceTick/*6*/);
+		ExpectedResult.Add(UninitTaskServiceCeaseRelevant/*7*/);
+		ExpectedResult.Add(SubServiceDeactivate/*5*/);
+		ExpectedResult.Add(Task2Execute/*9*/);
+		ExpectedResult.Add(ServiceDeactivate/*2*/);
 	}
 };
 IMPLEMENT_AI_LATENT_TEST(FAITest_BTStopDuringTaskServiceCeaseRelevant, "System.AI.Behavior Trees.Stop: during task service cease relevant")
 
+//////////////////////////////////////////////////////////////////////////////////////////////////
+// Uninitialize tests
+//////////////////////////////////////////////////////////////////////////////////////////////////
+struct FAITest_BTUninitializeDuringTaskExecute : public FAITest_SimpleBT
+{
+	FAITest_BTUninitializeDuringTaskExecute()
+	{
+	    enum
+	    {
+		    ServiceActivate = 1,
+		    ServiceDeactivate,
+		    ServiceTick,
+		    UninitTaskExecute,
+	    };
+
+		UBTCompositeNode& RootNode = FBTBuilder::AddSelector(*BTAsset);
+		UBTCompositeNode& CompNode = FBTBuilder::AddSelector(RootNode);
+		{
+			// Service
+			FBTBuilder::WithServiceLog(CompNode, ServiceActivate, ServiceDeactivate, ServiceTick);
+
+			UBTCompositeNode& Comp1Node = FBTBuilder::AddSelector(CompNode);
+			{
+				// UninitTask
+				FBTBuilder::AddTaskBTStopAction(Comp1Node, UninitTaskExecute, EBTNodeResult::Succeeded, EBTTestTaskStopTiming::DuringExecute, EBTTestStopAction::UnInitialize);
+			}
+		}
+
+		ExpectedResult.Add(ServiceActivate/*1*/);
+		ExpectedResult.Add(ServiceTick/*3*/);
+		ExpectedResult.Add(UninitTaskExecute/*4*/);
+		ExpectedResult.Add(ServiceDeactivate/*2*/);
+	}
+};
+IMPLEMENT_AI_LATENT_TEST(FAITest_BTUninitializeDuringTaskExecute, "System.AI.Behavior Trees.Uninitialize: during task execute")
+
+struct FAITest_BTUninitializeDuringTaskTick : public FAITest_SimpleBT
+{
+	FAITest_BTUninitializeDuringTaskTick()
+	{
+	    enum
+	    {
+		    ServiceActivate = 1,
+		    ServiceDeactivate,
+		    ServiceTick,
+		    UninitTaskTick,
+	    };
+
+		UBTCompositeNode& RootNode = FBTBuilder::AddSelector(*BTAsset);
+		UBTCompositeNode& CompNode = FBTBuilder::AddSelector(RootNode);
+		{
+			// Service
+			FBTBuilder::WithServiceLog(CompNode, ServiceActivate, ServiceDeactivate, ServiceTick);
+
+			UBTCompositeNode& Comp1Node = FBTBuilder::AddSelector(CompNode);
+			{
+				// UninitTask
+				FBTBuilder::AddTaskBTStopAction(Comp1Node, UninitTaskTick, EBTNodeResult::Succeeded, EBTTestTaskStopTiming::DuringTick, EBTTestStopAction::UnInitialize);
+			}
+		}
+
+		ExpectedResult.Add(ServiceActivate/*1*/);
+		ExpectedResult.Add(ServiceTick/*3*/);
+		ExpectedResult.Add(UninitTaskTick/*4*/);
+		ExpectedResult.Add(ServiceDeactivate/*2*/);
+	}
+};
+IMPLEMENT_AI_LATENT_TEST(FAITest_BTUninitializeDuringTaskTick, "System.AI.Behavior Trees.Uninitialize: during task tick")
+
+
+struct FAITest_BTUninitializeDuringTaskAbort : public FAITest_SimpleBT
+{
+	FAITest_BTUninitializeDuringTaskAbort()
+	{
+	    enum
+	    {
+		    ServiceActivate = 1,
+		    ServiceDeactivate,
+		    ServiceTick,
+		    SubServiceActivate,
+		    SubServiceDeactivate,
+		    SubServiceTick,
+		    UninitTaskAborting,
+			TaskExecute,
+	    };
+
+		UBTCompositeNode& RootNode = FBTBuilder::AddSelector(*BTAsset);
+		UBTCompositeNode& CompNode = FBTBuilder::AddSelector(RootNode);
+		{
+			// Service
+			FBTBuilder::WithServiceLog(CompNode, ServiceActivate, ServiceDeactivate, ServiceTick);
+
+			UBTCompositeNode& Comp1Node = FBTBuilder::AddSelector(CompNode);
+			{
+				// SubService
+				FBTBuilder::WithServiceLog(Comp1Node, SubServiceActivate, SubServiceDeactivate, SubServiceTick, TEXT("Bool1"));
+
+				// UninitTask
+				FBTBuilder::AddTaskBTStopAction(Comp1Node, UninitTaskAborting, EBTNodeResult::Succeeded, EBTTestTaskStopTiming::DuringAbort, EBTTestStopAction::UnInitialize);
+				FBTBuilder::WithDecoratorBlackboard(Comp1Node, EBasicKeyOperation::NotSet, EBTFlowAbortMode::Both, TEXT("Bool1"));
+
+				// Task
+				FBTBuilder::AddTask(Comp1Node, TaskExecute, EBTNodeResult::Succeeded);
+			}
+		}
+
+		ExpectedResult.Add(ServiceActivate/*1*/);
+		ExpectedResult.Add(SubServiceActivate/*4*/);
+		ExpectedResult.Add(ServiceTick/*5*/);
+		ExpectedResult.Add(SubServiceTick/*6*/);
+		ExpectedResult.Add(ServiceTick/*3*/);
+		ExpectedResult.Add(SubServiceTick/*6*/);
+		ExpectedResult.Add(UninitTaskAborting/*7*/);
+		ExpectedResult.Add(ServiceDeactivate/*2*/);
+		ExpectedResult.Add(SubServiceDeactivate/*5*/);
+	}
+};
+IMPLEMENT_AI_LATENT_TEST(FAITest_BTUninitializeDuringTaskAbort, "System.AI.Behavior Trees.Uninitialize: during task abort")
+
+
+struct FAITest_BTUninitializeDuringTaskFinished : public FAITest_SimpleBT
+{
+	FAITest_BTUninitializeDuringTaskFinished()
+	{
+	    enum
+	    {
+		    ServiceActivate = 1,
+		    ServiceDeactivate,
+		    ServiceTick,
+		    SubServiceActivate,
+		    SubServiceDeactivate,
+		    SubServiceTick,
+		    UninitTaskFinished,
+			TaskExecute,
+	    };
+
+		UBTCompositeNode& RootNode = FBTBuilder::AddSelector(*BTAsset);
+		UBTCompositeNode& CompNode = FBTBuilder::AddSelector(RootNode);
+		{
+			// Service
+			FBTBuilder::WithServiceLog(CompNode, ServiceActivate, ServiceDeactivate, ServiceTick);
+
+			UBTCompositeNode& Comp1Node = FBTBuilder::AddSelector(CompNode);
+			{
+				// SubService
+				FBTBuilder::WithServiceLog(Comp1Node, SubServiceActivate, SubServiceDeactivate, SubServiceTick, TEXT("Bool1"));
+
+				// UninitTask
+				FBTBuilder::AddTaskBTStopAction(Comp1Node, UninitTaskFinished, EBTNodeResult::Succeeded, EBTTestTaskStopTiming::DuringFinish, EBTTestStopAction::UnInitialize);
+				FBTBuilder::WithDecoratorBlackboard(Comp1Node, EBasicKeyOperation::NotSet, EBTFlowAbortMode::Both, TEXT("Bool1"));
+
+				// Task
+				FBTBuilder::AddTask(Comp1Node, TaskExecute, EBTNodeResult::Succeeded);
+			}
+		}
+
+		ExpectedResult.Add(ServiceActivate/*1*/);
+		ExpectedResult.Add(SubServiceActivate/*4*/);
+		ExpectedResult.Add(ServiceTick/*3*/);
+		ExpectedResult.Add(SubServiceTick/*6*/);
+		ExpectedResult.Add(ServiceTick/*3*/);
+		ExpectedResult.Add(SubServiceTick/*6*/);
+		ExpectedResult.Add(UninitTaskFinished/*7*/);
+		ExpectedResult.Add(ServiceDeactivate/*2*/);
+		ExpectedResult.Add(SubServiceDeactivate/*5*/);
+	}
+};
+IMPLEMENT_AI_LATENT_TEST(FAITest_BTUninitializeDuringTaskFinished, "System.AI.Behavior Trees.Uninitialize: during task finished")
+
+struct FAITest_BTUninitializeDuringServiceTick : public FAITest_SimpleBT
+{
+	FAITest_BTUninitializeDuringServiceTick()
+	{
+	    enum
+	    {
+		    ServiceActivate = 1,
+		    ServiceDeactivate,
+		    ServiceTick,
+		    SubServiceActivate,
+		    SubServiceDeactivate,
+		    SubServiceTick,
+		    UninitServiceTick,
+			TaskExecute,
+			Task2Execute,
+	    };
+
+		UBTCompositeNode& RootNode = FBTBuilder::AddSelector(*BTAsset);
+		UBTCompositeNode& CompNode = FBTBuilder::AddSelector(RootNode);
+		{
+			// Servier
+			FBTBuilder::WithServiceLog(CompNode, ServiceActivate, ServiceDeactivate, ServiceTick);
+
+			UBTCompositeNode& Comp1Node = FBTBuilder::AddSelector(CompNode);
+			{
+				// SubService
+				FBTBuilder::WithServiceLog(Comp1Node, SubServiceActivate, SubServiceDeactivate, SubServiceTick, TEXT("Bool1"));
+				// UninitServiceTick
+				FBTBuilder::WithServiceBTStopAction(Comp1Node, UninitServiceTick, EBTTestServiceStopTiming::DuringTick, EBTTestStopAction::UnInitialize);
+
+				// Task
+				FBTBuilder::AddTask(Comp1Node, TaskExecute, EBTNodeResult::Succeeded);
+				FBTBuilder::WithDecoratorBlackboard(Comp1Node, EBasicKeyOperation::NotSet, EBTFlowAbortMode::Both, TEXT("Bool1"));
+
+				// Task2
+				FBTBuilder::AddTask(Comp1Node, Task2Execute, EBTNodeResult::Succeeded);
+			}
+		}
+
+		ExpectedResult.Add(ServiceActivate/*1*/);
+		ExpectedResult.Add(SubServiceActivate/*4*/);
+		ExpectedResult.Add(ServiceTick/*3*/);
+		ExpectedResult.Add(SubServiceTick/*6*/);
+		ExpectedResult.Add(UninitServiceTick/*7*/);
+		ExpectedResult.Add(TaskExecute/*8*/);
+		ExpectedResult.Add(ServiceDeactivate/*2*/);
+		ExpectedResult.Add(SubServiceDeactivate/*5*/);
+	}
+};
+IMPLEMENT_AI_LATENT_TEST(FAITest_BTUninitializeDuringServiceTick, "System.AI.Behavior Trees.Uninitialize: during service tick")
+
+struct FAITest_BTUninitializeDuringServiceBecomeRelevant : public FAITest_SimpleBT
+{
+	FAITest_BTUninitializeDuringServiceBecomeRelevant()
+	{
+	    enum
+	    {
+		    ServiceActivate = 1,
+		    ServiceDeactivate,
+		    ServiceTick,
+		    SubServiceActivate,
+		    SubServiceDeactivate,
+		    SubServiceTick,
+		    UninitServiceBecomeRelevant,
+			TaskExecute,
+			Task2Execute,
+	    };
+
+		UBTCompositeNode& RootNode = FBTBuilder::AddSelector(*BTAsset);
+		UBTCompositeNode& CompNode = FBTBuilder::AddSelector(RootNode);
+		{
+			// Service
+			FBTBuilder::WithServiceLog(CompNode, ServiceActivate, ServiceDeactivate, ServiceTick);
+
+			UBTCompositeNode& Comp1Node = FBTBuilder::AddSelector(CompNode);
+			{
+				// SubService
+				FBTBuilder::WithServiceLog(Comp1Node, SubServiceActivate, SubServiceDeactivate, SubServiceTick, TEXT("Bool1"));
+				// UninitService
+				FBTBuilder::WithServiceBTStopAction(Comp1Node, UninitServiceBecomeRelevant, EBTTestServiceStopTiming::DuringBecomeRelevant, EBTTestStopAction::UnInitialize);
+
+				// Task
+				FBTBuilder::AddTask(Comp1Node, TaskExecute, EBTNodeResult::Succeeded);
+				FBTBuilder::WithDecoratorBlackboard(Comp1Node, EBasicKeyOperation::NotSet, EBTFlowAbortMode::Both, TEXT("Bool1"));
+
+				// Task2
+				FBTBuilder::AddTask(Comp1Node, Task2Execute, EBTNodeResult::Succeeded);
+			}
+		}
+
+		ExpectedResult.Add(ServiceActivate/*1*/);
+		ExpectedResult.Add(SubServiceActivate/*4*/);
+		ExpectedResult.Add(UninitServiceBecomeRelevant/*7*/);
+		ExpectedResult.Add(ServiceTick/*3*/);
+		ExpectedResult.Add(SubServiceTick/*6*/);
+		ExpectedResult.Add(TaskExecute/*8*/);
+		ExpectedResult.Add(ServiceDeactivate/*2*/);
+		ExpectedResult.Add(SubServiceDeactivate/*5*/);
+	}
+};
+IMPLEMENT_AI_LATENT_TEST(FAITest_BTUninitializeDuringServiceBecomeRelevant, "System.AI.Behavior Trees.Uninitialize: during service become relevant")
+
+struct FAITest_BTUninitializeDuringServiceCeaseRelevant : public FAITest_SimpleBT
+{
+	FAITest_BTUninitializeDuringServiceCeaseRelevant()
+	{
+	    enum
+	    {
+		    ServiceActivate = 1,
+		    ServiceDeactivate,
+		    ServiceTick,
+		    SubServiceActivate,
+		    SubServiceDeactivate,
+		    SubServiceTick,
+		    UninitServiceCeaseRelevant,
+			TaskExecute,
+			Task2Execute,
+	    };
+
+		UBTCompositeNode& RootNode = FBTBuilder::AddSelector(*BTAsset);
+		UBTCompositeNode& CompNode = FBTBuilder::AddSelector(RootNode);
+		{
+			// Service
+			FBTBuilder::WithServiceLog(CompNode, ServiceActivate, ServiceDeactivate, ServiceTick);
+
+			UBTCompositeNode& Comp1Node = FBTBuilder::AddSelector(CompNode);
+			{
+				// SubService
+				FBTBuilder::WithServiceLog(Comp1Node, SubServiceActivate, SubServiceDeactivate, SubServiceTick, TEXT("Bool1"));
+				// UninitService
+				FBTBuilder::WithServiceBTStopAction(Comp1Node, UninitServiceCeaseRelevant, EBTTestServiceStopTiming::DuringCeaseRelevant, EBTTestStopAction::UnInitialize);
+
+				// Task
+				FBTBuilder::AddTask(Comp1Node, TaskExecute, EBTNodeResult::Succeeded);
+				FBTBuilder::WithDecoratorBlackboard(Comp1Node, EBasicKeyOperation::NotSet, EBTFlowAbortMode::Both, TEXT("Bool1"));
+			}
+			// Task2
+			FBTBuilder::AddTask(CompNode, Task2Execute, EBTNodeResult::Succeeded);
+		}
+
+		ExpectedResult.Add(ServiceActivate/*1*/);
+		ExpectedResult.Add(SubServiceActivate/*4*/);
+		ExpectedResult.Add(ServiceTick/*3*/);
+		ExpectedResult.Add(SubServiceTick/*6*/);
+		ExpectedResult.Add(TaskExecute/*8*/);
+		ExpectedResult.Add(ServiceTick/*3*/);
+		ExpectedResult.Add(SubServiceTick/*6*/);
+		ExpectedResult.Add(SubServiceDeactivate/*5*/);
+		ExpectedResult.Add(UninitServiceCeaseRelevant/*7*/);
+		ExpectedResult.Add(Task2Execute/*9*/);
+		ExpectedResult.Add(ServiceDeactivate/*2*/);
+	}
+};
+IMPLEMENT_AI_LATENT_TEST(FAITest_BTUninitializeDuringServiceCeaseRelevant, "System.AI.Behavior Trees.Uninitialize: during service cease relevant")
+
+struct FAITest_BTUninitializeDuringTaskServiceTick : public FAITest_SimpleBT
+{
+	FAITest_BTUninitializeDuringTaskServiceTick()
+	{
+	    enum
+	    {
+		    ServiceActivate = 1,
+		    ServiceDeactivate,
+		    ServiceTick,
+		    SubServiceActivate,
+		    SubServiceDeactivate,
+		    SubServiceTick,
+		    UninitTaskServiceTick,
+			TaskExecute,
+			Task2Execute,
+	    };
+
+		UBTCompositeNode& RootNode = FBTBuilder::AddSelector(*BTAsset);
+		UBTCompositeNode& CompNode = FBTBuilder::AddSelector(RootNode);
+		{
+			FBTBuilder::WithServiceLog(CompNode, ServiceActivate, ServiceDeactivate, ServiceTick);
+
+			UBTCompositeNode& Comp1Node = FBTBuilder::AddSelector(CompNode);
+			{
+				FBTBuilder::WithServiceLog(Comp1Node, SubServiceActivate, SubServiceDeactivate, SubServiceTick, TEXT("Bool1"));
+
+				FBTBuilder::AddTask(Comp1Node, TaskExecute, EBTNodeResult::Succeeded);
+				FBTBuilder::WithDecoratorBlackboard(Comp1Node, EBasicKeyOperation::NotSet, EBTFlowAbortMode::Both, TEXT("Bool1"));
+				FBTBuilder::WithTaskServiceBTStopAction(Comp1Node, UninitTaskServiceTick, EBTTestServiceStopTiming::DuringTick, EBTTestStopAction::UnInitialize);
+
+				FBTBuilder::AddTask(Comp1Node, Task2Execute, EBTNodeResult::Succeeded);
+			}
+		}
+
+		ExpectedResult.Add(ServiceActivate/*1*/);
+		ExpectedResult.Add(SubServiceActivate/*4*/);
+		ExpectedResult.Add(ServiceTick/*3*/);
+		ExpectedResult.Add(SubServiceTick/*6*/);
+		ExpectedResult.Add(UninitTaskServiceTick/*7*/);
+		ExpectedResult.Add(TaskExecute/*8*/);
+		ExpectedResult.Add(ServiceDeactivate/*2*/);
+		ExpectedResult.Add(SubServiceDeactivate/*5*/);
+	}
+};
+IMPLEMENT_AI_LATENT_TEST(FAITest_BTUninitializeDuringTaskServiceTick, "System.AI.Behavior Trees.Uninitialize: during task service tick")
+
+struct FAITest_BTUninitializeDuringTaskServiceBecomeRelevant : public FAITest_SimpleBT
+{
+	FAITest_BTUninitializeDuringTaskServiceBecomeRelevant()
+	{
+	    enum
+	    {
+		    ServiceActivate = 1,
+		    ServiceDeactivate,
+		    ServiceTick,
+		    SubServiceActivate,
+		    SubServiceDeactivate,
+		    SubServiceTick,
+		    UninitTaskServiceBecomeRelevant,
+			TaskExecute,
+			Task2Execute,
+	    };
+
+		UBTCompositeNode& RootNode = FBTBuilder::AddSelector(*BTAsset);
+		UBTCompositeNode& CompNode = FBTBuilder::AddSelector(RootNode);
+		{
+			FBTBuilder::WithServiceLog(CompNode, ServiceActivate, ServiceDeactivate, ServiceTick);
+
+			UBTCompositeNode& Comp1Node = FBTBuilder::AddSelector(CompNode);
+			{
+				FBTBuilder::WithServiceLog(Comp1Node, SubServiceActivate, SubServiceDeactivate, SubServiceTick, TEXT("Bool1"));
+
+				FBTBuilder::AddTask(Comp1Node, TaskExecute, EBTNodeResult::Succeeded);
+				FBTBuilder::WithDecoratorBlackboard(Comp1Node, EBasicKeyOperation::NotSet, EBTFlowAbortMode::Both, TEXT("Bool1"));
+				FBTBuilder::WithTaskServiceBTStopAction(Comp1Node, UninitTaskServiceBecomeRelevant, EBTTestServiceStopTiming::DuringBecomeRelevant, EBTTestStopAction::UnInitialize);
+
+				FBTBuilder::AddTask(Comp1Node, Task2Execute, EBTNodeResult::Succeeded);
+			}
+		}
+
+		ExpectedResult.Add(ServiceActivate/*1*/);
+		ExpectedResult.Add(SubServiceActivate/*4*/);
+		ExpectedResult.Add(ServiceTick/*3*/);
+		ExpectedResult.Add(SubServiceTick/*6*/);
+		ExpectedResult.Add(UninitTaskServiceBecomeRelevant/*7*/);
+		ExpectedResult.Add(TaskExecute/*8*/);
+		ExpectedResult.Add(ServiceDeactivate/*2*/);
+		ExpectedResult.Add(SubServiceDeactivate/*5*/);
+	}
+};
+IMPLEMENT_AI_LATENT_TEST(FAITest_BTUninitializeDuringTaskServiceBecomeRelevant, "System.AI.Behavior Trees.Uninitialize: during task service become relevant")
+
+
+struct FAITest_BTUninitializeDuringTaskServiceCeaseRelevant : public FAITest_SimpleBT
+{
+	FAITest_BTUninitializeDuringTaskServiceCeaseRelevant()
+	{
+	    enum
+	    {
+		    ServiceActivate = 1,
+		    ServiceDeactivate,
+		    ServiceTick,
+		    SubServiceActivate,
+		    SubServiceDeactivate,
+		    SubServiceTick,
+		    UninitTaskServiceCeaseRelevant,
+			TaskExecute,
+			Task2Execute,
+	    };
+
+		UBTCompositeNode& RootNode = FBTBuilder::AddSelector(*BTAsset);
+		UBTCompositeNode& CompNode = FBTBuilder::AddSelector(RootNode);
+		{
+			FBTBuilder::WithServiceLog(CompNode, ServiceActivate, ServiceDeactivate, ServiceTick);
+
+			UBTCompositeNode& Comp1Node = FBTBuilder::AddSelector(CompNode);
+			{
+				FBTBuilder::WithServiceLog(Comp1Node, SubServiceActivate, SubServiceDeactivate, SubServiceTick, TEXT("Bool1"));
+
+				FBTBuilder::AddTask(Comp1Node, TaskExecute, EBTNodeResult::Succeeded);
+				FBTBuilder::WithDecoratorBlackboard(Comp1Node, EBasicKeyOperation::NotSet, EBTFlowAbortMode::Both, TEXT("Bool1"));
+				FBTBuilder::WithTaskServiceBTStopAction(Comp1Node, UninitTaskServiceCeaseRelevant, EBTTestServiceStopTiming::DuringCeaseRelevant, EBTTestStopAction::UnInitialize);
+
+				FBTBuilder::AddTask(CompNode, Task2Execute, EBTNodeResult::Succeeded);
+			}
+		}
+
+		ExpectedResult.Add(ServiceActivate/*1*/);
+		ExpectedResult.Add(SubServiceActivate/*4*/);
+		ExpectedResult.Add(ServiceTick/*3*/);
+		ExpectedResult.Add(SubServiceTick/*6*/);
+		ExpectedResult.Add(TaskExecute/*8*/);
+		ExpectedResult.Add(ServiceTick/*3*/);
+		ExpectedResult.Add(SubServiceTick/*6*/);
+		ExpectedResult.Add(UninitTaskServiceCeaseRelevant/*7*/);
+		ExpectedResult.Add(SubServiceDeactivate/*5*/);
+		ExpectedResult.Add(Task2Execute/*9*/);
+		ExpectedResult.Add(ServiceDeactivate/*2*/);
+	}
+};
+IMPLEMENT_AI_LATENT_TEST(FAITest_BTUninitializeDuringTaskServiceCeaseRelevant, "System.AI.Behavior Trees.Uninitialize: during task service cease relevant")
+
+
+//////////////////////////////////////////////////////////////////////////////////////////////////
+// Cleanup tests
+//////////////////////////////////////////////////////////////////////////////////////////////////
+struct FAITest_BTCleanupDuringTaskExecute : public FAITest_SimpleBT
+{
+	FAITest_BTCleanupDuringTaskExecute()
+	{
+		enum
+		{
+			ServiceActivate = 1,
+			ServiceDeactivate,
+			ServiceTick,
+			UninitTaskExecute,
+		};
+
+		UBTCompositeNode& RootNode = FBTBuilder::AddSelector(*BTAsset);
+		UBTCompositeNode& CompNode = FBTBuilder::AddSelector(RootNode);
+		{
+			// Service
+			FBTBuilder::WithServiceLog(CompNode, ServiceActivate, ServiceDeactivate, ServiceTick);
+
+			UBTCompositeNode& Comp1Node = FBTBuilder::AddSelector(CompNode);
+			{
+				// UninitTask
+				FBTBuilder::AddTaskBTStopAction(Comp1Node, UninitTaskExecute, EBTNodeResult::Succeeded, EBTTestTaskStopTiming::DuringExecute, EBTTestStopAction::Cleanup);
+			}
+		}
+
+		ExpectedResult.Add(ServiceActivate/*1*/);
+		ExpectedResult.Add(ServiceTick/*3*/);
+		ExpectedResult.Add(UninitTaskExecute/*4*/);
+		ExpectedResult.Add(ServiceDeactivate/*2*/);
+	}
+};
+IMPLEMENT_AI_LATENT_TEST(FAITest_BTCleanupDuringTaskExecute, "System.AI.Behavior Trees.Cleanup: during task execute")
+
+struct FAITest_BTCleanupDuringTaskTick : public FAITest_SimpleBT
+{
+	FAITest_BTCleanupDuringTaskTick()
+	{
+		enum
+		{
+			ServiceActivate = 1,
+			ServiceDeactivate,
+			ServiceTick,
+			UninitTaskTick,
+		};
+
+		UBTCompositeNode& RootNode = FBTBuilder::AddSelector(*BTAsset);
+		UBTCompositeNode& CompNode = FBTBuilder::AddSelector(RootNode);
+		{
+			// Service
+			FBTBuilder::WithServiceLog(CompNode, ServiceActivate, ServiceDeactivate, ServiceTick);
+
+			UBTCompositeNode& Comp1Node = FBTBuilder::AddSelector(CompNode);
+			{
+				// UninitTask
+				FBTBuilder::AddTaskBTStopAction(Comp1Node, UninitTaskTick, EBTNodeResult::Succeeded, EBTTestTaskStopTiming::DuringTick, EBTTestStopAction::Cleanup);
+			}
+		}
+
+		ExpectedResult.Add(ServiceActivate/*1*/);
+		ExpectedResult.Add(ServiceTick/*3*/);
+		ExpectedResult.Add(UninitTaskTick/*4*/);
+		ExpectedResult.Add(ServiceDeactivate/*2*/);
+	}
+};
+IMPLEMENT_AI_LATENT_TEST(FAITest_BTCleanupDuringTaskTick, "System.AI.Behavior Trees.Cleanup: during task tick")
+
+
+struct FAITest_BTCleanupDuringTaskAbort : public FAITest_SimpleBT
+{
+	FAITest_BTCleanupDuringTaskAbort()
+	{
+		enum
+		{
+			ServiceActivate = 1,
+			ServiceDeactivate,
+			ServiceTick,
+			SubServiceActivate,
+			SubServiceDeactivate,
+			SubServiceTick,
+			UninitTaskAborting,
+			TaskExecute,
+		};
+
+		UBTCompositeNode& RootNode = FBTBuilder::AddSelector(*BTAsset);
+		UBTCompositeNode& CompNode = FBTBuilder::AddSelector(RootNode);
+		{
+			// Service
+			FBTBuilder::WithServiceLog(CompNode, ServiceActivate, ServiceDeactivate, ServiceTick);
+
+			UBTCompositeNode& Comp1Node = FBTBuilder::AddSelector(CompNode);
+			{
+				// SubService
+				FBTBuilder::WithServiceLog(Comp1Node, SubServiceActivate, SubServiceDeactivate, SubServiceTick, TEXT("Bool1"));
+
+				// UninitTask
+				FBTBuilder::AddTaskBTStopAction(Comp1Node, UninitTaskAborting, EBTNodeResult::Succeeded, EBTTestTaskStopTiming::DuringAbort, EBTTestStopAction::Cleanup);
+				FBTBuilder::WithDecoratorBlackboard(Comp1Node, EBasicKeyOperation::NotSet, EBTFlowAbortMode::Both, TEXT("Bool1"));
+
+				// Task
+				FBTBuilder::AddTask(Comp1Node, TaskExecute, EBTNodeResult::Succeeded);
+			}
+		}
+
+		ExpectedResult.Add(ServiceActivate/*1*/);
+		ExpectedResult.Add(SubServiceActivate/*4*/);
+		ExpectedResult.Add(ServiceTick/*5*/);
+		ExpectedResult.Add(SubServiceTick/*6*/);
+		ExpectedResult.Add(ServiceTick/*3*/);
+		ExpectedResult.Add(SubServiceTick/*6*/);
+		ExpectedResult.Add(UninitTaskAborting/*7*/);
+		ExpectedResult.Add(ServiceDeactivate/*2*/);
+		ExpectedResult.Add(SubServiceDeactivate/*5*/);
+	}
+};
+IMPLEMENT_AI_LATENT_TEST(FAITest_BTCleanupDuringTaskAbort, "System.AI.Behavior Trees.Cleanup: during task abort")
+
+
+struct FAITest_BTCleanupDuringTaskFinished : public FAITest_SimpleBT
+{
+	FAITest_BTCleanupDuringTaskFinished()
+	{
+		enum
+		{
+			ServiceActivate = 1,
+			ServiceDeactivate,
+			ServiceTick,
+			SubServiceActivate,
+			SubServiceDeactivate,
+			SubServiceTick,
+			UninitTaskFinished,
+			TaskExecute,
+		};
+
+		UBTCompositeNode& RootNode = FBTBuilder::AddSelector(*BTAsset);
+		UBTCompositeNode& CompNode = FBTBuilder::AddSelector(RootNode);
+		{
+			// Service
+			FBTBuilder::WithServiceLog(CompNode, ServiceActivate, ServiceDeactivate, ServiceTick);
+
+			UBTCompositeNode& Comp1Node = FBTBuilder::AddSelector(CompNode);
+			{
+				// SubService
+				FBTBuilder::WithServiceLog(Comp1Node, SubServiceActivate, SubServiceDeactivate, SubServiceTick, TEXT("Bool1"));
+
+				// UninitTask
+				FBTBuilder::AddTaskBTStopAction(Comp1Node, UninitTaskFinished, EBTNodeResult::Succeeded, EBTTestTaskStopTiming::DuringFinish, EBTTestStopAction::Cleanup);
+				FBTBuilder::WithDecoratorBlackboard(Comp1Node, EBasicKeyOperation::NotSet, EBTFlowAbortMode::Both, TEXT("Bool1"));
+
+				// Task
+				FBTBuilder::AddTask(Comp1Node, TaskExecute, EBTNodeResult::Succeeded);
+			}
+		}
+
+		ExpectedResult.Add(ServiceActivate/*1*/);
+		ExpectedResult.Add(SubServiceActivate/*4*/);
+		ExpectedResult.Add(ServiceTick/*3*/);
+		ExpectedResult.Add(SubServiceTick/*6*/);
+		ExpectedResult.Add(ServiceTick/*3*/);
+		ExpectedResult.Add(SubServiceTick/*6*/);
+		ExpectedResult.Add(UninitTaskFinished/*7*/);
+		ExpectedResult.Add(ServiceDeactivate/*2*/);
+		ExpectedResult.Add(SubServiceDeactivate/*5*/);
+	}
+};
+IMPLEMENT_AI_LATENT_TEST(FAITest_BTCleanupDuringTaskFinished, "System.AI.Behavior Trees.Cleanup: during task finished")
+
+struct FAITest_BTCleanupDuringServiceTick : public FAITest_SimpleBT
+{
+	FAITest_BTCleanupDuringServiceTick()
+	{
+		enum
+		{
+			ServiceActivate = 1,
+			ServiceDeactivate,
+			ServiceTick,
+			SubServiceActivate,
+			SubServiceDeactivate,
+			SubServiceTick,
+			UninitServiceTick,
+			TaskExecute,
+			Task2Execute,
+		};
+
+		UBTCompositeNode& RootNode = FBTBuilder::AddSelector(*BTAsset);
+		UBTCompositeNode& CompNode = FBTBuilder::AddSelector(RootNode);
+		{
+			// Servier
+			FBTBuilder::WithServiceLog(CompNode, ServiceActivate, ServiceDeactivate, ServiceTick);
+
+			UBTCompositeNode& Comp1Node = FBTBuilder::AddSelector(CompNode);
+			{
+				// SubService
+				FBTBuilder::WithServiceLog(Comp1Node, SubServiceActivate, SubServiceDeactivate, SubServiceTick, TEXT("Bool1"));
+				// UninitServiceTick
+				FBTBuilder::WithServiceBTStopAction(Comp1Node, UninitServiceTick, EBTTestServiceStopTiming::DuringTick, EBTTestStopAction::Cleanup);
+
+				// Task
+				FBTBuilder::AddTask(Comp1Node, TaskExecute, EBTNodeResult::Succeeded);
+				FBTBuilder::WithDecoratorBlackboard(Comp1Node, EBasicKeyOperation::NotSet, EBTFlowAbortMode::Both, TEXT("Bool1"));
+
+				// Task2
+				FBTBuilder::AddTask(Comp1Node, Task2Execute, EBTNodeResult::Succeeded);
+			}
+		}
+
+		ExpectedResult.Add(ServiceActivate/*1*/);
+		ExpectedResult.Add(SubServiceActivate/*4*/);
+		ExpectedResult.Add(ServiceTick/*3*/);
+		ExpectedResult.Add(SubServiceTick/*6*/);
+		ExpectedResult.Add(UninitServiceTick/*7*/);
+		ExpectedResult.Add(TaskExecute/*8*/);
+		ExpectedResult.Add(ServiceDeactivate/*2*/);
+		ExpectedResult.Add(SubServiceDeactivate/*5*/);
+	}
+};
+IMPLEMENT_AI_LATENT_TEST(FAITest_BTCleanupDuringServiceTick, "System.AI.Behavior Trees.Cleanup: during service tick")
+
+struct FAITest_BTCleanupDuringServiceBecomeRelevant : public FAITest_SimpleBT
+{
+	FAITest_BTCleanupDuringServiceBecomeRelevant()
+	{
+		enum
+		{
+			ServiceActivate = 1,
+			ServiceDeactivate,
+			ServiceTick,
+			SubServiceActivate,
+			SubServiceDeactivate,
+			SubServiceTick,
+			UninitServiceBecomeRelevant,
+			TaskExecute,
+			Task2Execute,
+		};
+
+		UBTCompositeNode& RootNode = FBTBuilder::AddSelector(*BTAsset);
+		UBTCompositeNode& CompNode = FBTBuilder::AddSelector(RootNode);
+		{
+			// Service
+			FBTBuilder::WithServiceLog(CompNode, ServiceActivate, ServiceDeactivate, ServiceTick);
+
+			UBTCompositeNode& Comp1Node = FBTBuilder::AddSelector(CompNode);
+			{
+				// SubService
+				FBTBuilder::WithServiceLog(Comp1Node, SubServiceActivate, SubServiceDeactivate, SubServiceTick, TEXT("Bool1"));
+				// UninitService
+				FBTBuilder::WithServiceBTStopAction(Comp1Node, UninitServiceBecomeRelevant, EBTTestServiceStopTiming::DuringBecomeRelevant, EBTTestStopAction::Cleanup);
+
+				// Task
+				FBTBuilder::AddTask(Comp1Node, TaskExecute, EBTNodeResult::Succeeded);
+				FBTBuilder::WithDecoratorBlackboard(Comp1Node, EBasicKeyOperation::NotSet, EBTFlowAbortMode::Both, TEXT("Bool1"));
+
+				// Task2
+				FBTBuilder::AddTask(Comp1Node, Task2Execute, EBTNodeResult::Succeeded);
+			}
+		}
+
+		ExpectedResult.Add(ServiceActivate/*1*/);
+		ExpectedResult.Add(SubServiceActivate/*4*/);
+		ExpectedResult.Add(UninitServiceBecomeRelevant/*7*/);
+		ExpectedResult.Add(ServiceTick/*3*/);
+		ExpectedResult.Add(SubServiceTick/*6*/);
+		ExpectedResult.Add(TaskExecute/*8*/);
+		ExpectedResult.Add(ServiceDeactivate/*2*/);
+		ExpectedResult.Add(SubServiceDeactivate/*5*/);
+	}
+};
+IMPLEMENT_AI_LATENT_TEST(FAITest_BTCleanupDuringServiceBecomeRelevant, "System.AI.Behavior Trees.Cleanup: during service become relevant")
+
+struct FAITest_BTCleanupDuringServiceCeaseRelevant : public FAITest_SimpleBT
+{
+	FAITest_BTCleanupDuringServiceCeaseRelevant()
+	{
+		enum
+		{
+			ServiceActivate = 1,
+			ServiceDeactivate,
+			ServiceTick,
+			SubServiceActivate,
+			SubServiceDeactivate,
+			SubServiceTick,
+			UninitServiceCeaseRelevant,
+			TaskExecute,
+			Task2Execute,
+		};
+
+		UBTCompositeNode& RootNode = FBTBuilder::AddSelector(*BTAsset);
+		UBTCompositeNode& CompNode = FBTBuilder::AddSelector(RootNode);
+		{
+			// Service
+			FBTBuilder::WithServiceLog(CompNode, ServiceActivate, ServiceDeactivate, ServiceTick);
+
+			UBTCompositeNode& Comp1Node = FBTBuilder::AddSelector(CompNode);
+			{
+				// SubService
+				FBTBuilder::WithServiceLog(Comp1Node, SubServiceActivate, SubServiceDeactivate, SubServiceTick, TEXT("Bool1"));
+				// UninitService
+				FBTBuilder::WithServiceBTStopAction(Comp1Node, UninitServiceCeaseRelevant, EBTTestServiceStopTiming::DuringCeaseRelevant, EBTTestStopAction::Cleanup);
+
+				// Task
+				FBTBuilder::AddTask(Comp1Node, TaskExecute, EBTNodeResult::Succeeded);
+				FBTBuilder::WithDecoratorBlackboard(Comp1Node, EBasicKeyOperation::NotSet, EBTFlowAbortMode::Both, TEXT("Bool1"));
+			}
+			// Task2
+			FBTBuilder::AddTask(CompNode, Task2Execute, EBTNodeResult::Succeeded);
+		}
+
+		ExpectedResult.Add(ServiceActivate/*1*/);
+		ExpectedResult.Add(SubServiceActivate/*4*/);
+		ExpectedResult.Add(ServiceTick/*3*/);
+		ExpectedResult.Add(SubServiceTick/*6*/);
+		ExpectedResult.Add(TaskExecute/*8*/);
+		ExpectedResult.Add(ServiceTick/*3*/);
+		ExpectedResult.Add(SubServiceTick/*6*/);
+		ExpectedResult.Add(SubServiceDeactivate/*5*/);
+		ExpectedResult.Add(UninitServiceCeaseRelevant/*7*/);
+		ExpectedResult.Add(Task2Execute/*9*/);
+		ExpectedResult.Add(ServiceDeactivate/*2*/);
+	}
+};
+IMPLEMENT_AI_LATENT_TEST(FAITest_BTCleanupDuringServiceCeaseRelevant, "System.AI.Behavior Trees.Cleanup: during service cease relevant")
+
+struct FAITest_BTCleanupDuringTaskServiceTick : public FAITest_SimpleBT
+{
+	FAITest_BTCleanupDuringTaskServiceTick()
+	{
+		enum
+		{
+			ServiceActivate = 1,
+			ServiceDeactivate,
+			ServiceTick,
+			SubServiceActivate,
+			SubServiceDeactivate,
+			SubServiceTick,
+			UninitTaskServiceTick,
+			TaskExecute,
+			Task2Execute,
+		};
+
+		UBTCompositeNode& RootNode = FBTBuilder::AddSelector(*BTAsset);
+		UBTCompositeNode& CompNode = FBTBuilder::AddSelector(RootNode);
+		{
+			FBTBuilder::WithServiceLog(CompNode, ServiceActivate, ServiceDeactivate, ServiceTick);
+
+			UBTCompositeNode& Comp1Node = FBTBuilder::AddSelector(CompNode);
+			{
+				FBTBuilder::WithServiceLog(Comp1Node, SubServiceActivate, SubServiceDeactivate, SubServiceTick, TEXT("Bool1"));
+
+				FBTBuilder::AddTask(Comp1Node, TaskExecute, EBTNodeResult::Succeeded);
+				FBTBuilder::WithDecoratorBlackboard(Comp1Node, EBasicKeyOperation::NotSet, EBTFlowAbortMode::Both, TEXT("Bool1"));
+				FBTBuilder::WithTaskServiceBTStopAction(Comp1Node, UninitTaskServiceTick, EBTTestServiceStopTiming::DuringTick, EBTTestStopAction::Cleanup);
+
+				FBTBuilder::AddTask(Comp1Node, Task2Execute, EBTNodeResult::Succeeded);
+			}
+		}
+
+		ExpectedResult.Add(ServiceActivate/*1*/);
+		ExpectedResult.Add(SubServiceActivate/*4*/);
+		ExpectedResult.Add(ServiceTick/*3*/);
+		ExpectedResult.Add(SubServiceTick/*6*/);
+		ExpectedResult.Add(UninitTaskServiceTick/*7*/);
+		ExpectedResult.Add(TaskExecute/*8*/);
+		ExpectedResult.Add(ServiceDeactivate/*2*/);
+		ExpectedResult.Add(SubServiceDeactivate/*5*/);
+	}
+};
+IMPLEMENT_AI_LATENT_TEST(FAITest_BTCleanupDuringTaskServiceTick, "System.AI.Behavior Trees.Cleanup: during task service tick")
+
+struct FAITest_BTCleanupDuringTaskServiceBecomeRelevant : public FAITest_SimpleBT
+{
+	FAITest_BTCleanupDuringTaskServiceBecomeRelevant()
+	{
+		enum
+		{
+			ServiceActivate = 1,
+			ServiceDeactivate,
+			ServiceTick,
+			SubServiceActivate,
+			SubServiceDeactivate,
+			SubServiceTick,
+			UninitTaskServiceBecomeRelevant,
+			TaskExecute,
+			Task2Execute,
+		};
+
+		UBTCompositeNode& RootNode = FBTBuilder::AddSelector(*BTAsset);
+		UBTCompositeNode& CompNode = FBTBuilder::AddSelector(RootNode);
+		{
+			FBTBuilder::WithServiceLog(CompNode, ServiceActivate, ServiceDeactivate, ServiceTick);
+
+			UBTCompositeNode& Comp1Node = FBTBuilder::AddSelector(CompNode);
+			{
+				FBTBuilder::WithServiceLog(Comp1Node, SubServiceActivate, SubServiceDeactivate, SubServiceTick, TEXT("Bool1"));
+
+				FBTBuilder::AddTask(Comp1Node, TaskExecute, EBTNodeResult::Succeeded);
+				FBTBuilder::WithDecoratorBlackboard(Comp1Node, EBasicKeyOperation::NotSet, EBTFlowAbortMode::Both, TEXT("Bool1"));
+				FBTBuilder::WithTaskServiceBTStopAction(Comp1Node, UninitTaskServiceBecomeRelevant, EBTTestServiceStopTiming::DuringBecomeRelevant, EBTTestStopAction::Cleanup);
+
+				FBTBuilder::AddTask(Comp1Node, Task2Execute, EBTNodeResult::Succeeded);
+			}
+		}
+
+		ExpectedResult.Add(ServiceActivate/*1*/);
+		ExpectedResult.Add(SubServiceActivate/*4*/);
+		ExpectedResult.Add(ServiceTick/*3*/);
+		ExpectedResult.Add(SubServiceTick/*6*/);
+		ExpectedResult.Add(UninitTaskServiceBecomeRelevant/*7*/);
+		ExpectedResult.Add(TaskExecute/*8*/);
+		ExpectedResult.Add(ServiceDeactivate/*2*/);
+		ExpectedResult.Add(SubServiceDeactivate/*5*/);
+	}
+};
+IMPLEMENT_AI_LATENT_TEST(FAITest_BTCleanupDuringTaskServiceBecomeRelevant, "System.AI.Behavior Trees.Cleanup: during task service become relevant")
+
+
+struct FAITest_BTCleanupDuringTaskServiceCeaseRelevant : public FAITest_SimpleBT
+{
+	FAITest_BTCleanupDuringTaskServiceCeaseRelevant()
+	{
+		enum
+		{
+			ServiceActivate = 1,
+			ServiceDeactivate,
+			ServiceTick,
+			SubServiceActivate,
+			SubServiceDeactivate,
+			SubServiceTick,
+			UninitTaskServiceCeaseRelevant,
+			TaskExecute,
+			Task2Execute,
+		};
+
+		UBTCompositeNode& RootNode = FBTBuilder::AddSelector(*BTAsset);
+		UBTCompositeNode& CompNode = FBTBuilder::AddSelector(RootNode);
+		{
+			FBTBuilder::WithServiceLog(CompNode, ServiceActivate, ServiceDeactivate, ServiceTick);
+
+			UBTCompositeNode& Comp1Node = FBTBuilder::AddSelector(CompNode);
+			{
+				FBTBuilder::WithServiceLog(Comp1Node, SubServiceActivate, SubServiceDeactivate, SubServiceTick, TEXT("Bool1"));
+
+				FBTBuilder::AddTask(Comp1Node, TaskExecute, EBTNodeResult::Succeeded);
+				FBTBuilder::WithDecoratorBlackboard(Comp1Node, EBasicKeyOperation::NotSet, EBTFlowAbortMode::Both, TEXT("Bool1"));
+				FBTBuilder::WithTaskServiceBTStopAction(Comp1Node, UninitTaskServiceCeaseRelevant, EBTTestServiceStopTiming::DuringCeaseRelevant, EBTTestStopAction::Cleanup);
+
+				FBTBuilder::AddTask(CompNode, Task2Execute, EBTNodeResult::Succeeded);
+			}
+		}
+
+		ExpectedResult.Add(ServiceActivate/*1*/);
+		ExpectedResult.Add(SubServiceActivate/*4*/);
+		ExpectedResult.Add(ServiceTick/*3*/);
+		ExpectedResult.Add(SubServiceTick/*6*/);
+		ExpectedResult.Add(TaskExecute/*8*/);
+		ExpectedResult.Add(ServiceTick/*3*/);
+		ExpectedResult.Add(SubServiceTick/*6*/);
+		ExpectedResult.Add(UninitTaskServiceCeaseRelevant/*7*/);
+		ExpectedResult.Add(SubServiceDeactivate/*5*/);
+		ExpectedResult.Add(Task2Execute/*9*/);
+		ExpectedResult.Add(ServiceDeactivate/*2*/);
+	}
+};
+IMPLEMENT_AI_LATENT_TEST(FAITest_BTCleanupDuringTaskServiceCeaseRelevant, "System.AI.Behavior Trees.Cleanup: during task service cease relevant")
+
 #undef LOCTEXT_NAMESPACE
\ No newline at end of file
diff --git a/Engine/Source/Developer/AITestSuite/Public/BTBuilder.h b/Engine/Source/Developer/AITestSuite/Public/BTBuilder.h
--- a/Engine/Source/Developer/AITestSuite/Public/BTBuilder.h
+++ b/Engine/Source/Developer/AITestSuite/Public/BTBuilder.h
@@ -18,12 +18,12 @@
 #include "BehaviorTree/TestBTDecorator_Blackboard.h"
 #include "BehaviorTree/TestBTDecorator_DelayedAbort.h"
 #include "BehaviorTree/TestBTService_Log.h"
-#include "BehaviorTree/TestBTService_StopTree.h"
+#include "BehaviorTree/TestBTService_BTStopAction.h"
 #include "BehaviorTree/TestBTTask_LatentWithFlags.h"
 #include "BehaviorTree/TestBTTask_Log.h"
 #include "BehaviorTree/TestBTTask_SetFlag.h"
 #include "BehaviorTree/TestBTTask_SetValue.h"
-#include "BehaviorTree/TestBTTask_StopTree.h"
+#include "BehaviorTree/TestBTTask_BTStopAction.h"
 #include "BehaviorTree/TestBTTask_ToggleFlag.h"
 
 struct FBTBuilder
@@ -221,12 +221,13 @@ struct FBTBuilder
 		ParentNode.Children[ChildIdx].ChildTask = TaskNode;
 	}
 
-	static void AddTaskStopTree(UBTCompositeNode& ParentNode, int32 LogIndex, EBTNodeResult::Type NodeResult, EBTTestTaskStopTree::Type StopTimming)
+	static void AddTaskBTStopAction(UBTCompositeNode& ParentNode, int32 LogIndex, EBTNodeResult::Type NodeResult, EBTTestTaskStopTiming StopTiming, EBTTestStopAction StopAction)
 	{
-		UTestBTTask_StopTree* TaskNode = NewObject<UTestBTTask_StopTree>(ParentNode.GetTreeAsset());
+		UTestBTTask_BTStopAction* TaskNode = NewObject<UTestBTTask_BTStopAction>(ParentNode.GetTreeAsset());
 		TaskNode->LogIndex = LogIndex;
 		TaskNode->LogResult = NodeResult;
-		TaskNode->StopTimming = StopTimming;
+		TaskNode->StopTiming = StopTiming;
+		TaskNode->StopAction = StopAction;
 
 		const int32 ChildIdx = ParentNode.Children.AddZeroed(1);
 		ParentNode.Children[ChildIdx].ChildTask = TaskNode;
@@ -327,11 +328,12 @@ struct FBTBuilder
 		LogService.bToggleValue = bToggleValue;
 	}
 
-	static void WithServiceStopTree(UBTCompositeNode& ParentNode, int32 LogIndex, EBTTestServiceStopTree::Type StopTimming)
+	static void WithServiceBTStopAction(UBTCompositeNode& ParentNode, int32 LogIndex, EBTTestServiceStopTiming StopTiming, EBTTestStopAction StopAction)
 	{
-		UTestBTService_StopTree& Service = WithService<UTestBTService_StopTree>(ParentNode);
+		UTestBTService_BTStopAction& Service = WithService<UTestBTService_BTStopAction>(ParentNode);
 		Service.LogIndex = LogIndex;
-		Service.StopTimming = StopTimming;
+		Service.StopTiming = StopTiming;
+		Service.StopAction = StopAction;
 	}
 
 	template<class T>
@@ -358,11 +360,12 @@ struct FBTBuilder
 		LogService.bToggleValue = bToggleValue;
 	}
 
-	static void WithTaskServiceStopTree(UBTCompositeNode& ParentNode, int32 LogIndex, EBTTestServiceStopTree::Type StopTimming)
+	static void WithTaskServiceBTStopAction(UBTCompositeNode& ParentNode, int32 LogIndex, EBTTestServiceStopTiming StopTiming, EBTTestStopAction StopAction)
 	{
-		UTestBTService_StopTree& Service = WithTaskService<UTestBTService_StopTree>(ParentNode);
+		UTestBTService_BTStopAction& Service = WithTaskService<UTestBTService_BTStopAction>(ParentNode);
 		Service.LogIndex = LogIndex;
-		Service.StopTimming = StopTimming;
+		Service.StopTiming = StopTiming;
+		Service.StopAction = StopAction;
 	}
 
 };
