Fix inconsistencies in AITests FTestHelpers functions return

diff --git a/Engine/Source/Developer/AITestSuite/Public/AITestsCommon.h b/Engine/Source/Developer/AITestSuite/Public/AITestsCommon.h
--- a/Engine/Source/Developer/AITestSuite/Public/AITestsCommon.h
+++ b/Engine/Source/Developer/AITestSuite/Public/AITestsCommon.h
@@ -211,29 +211,25 @@ namespace FTestHelpers
 	template<typename T1, typename T2>
 	inline bool TestEqual(const FString& Description, T1 Expression, T2 Expected, FAutomationTestBase& This)
 	{
-		This.TestEqual(*Description, Expression, Expected);
-		return Expression == Expected;
+		return This.TestEqual(*Description, Expression, Expected);
 	}
 
 	template<typename T1, typename T2>
 	inline bool TestEqual(const FString& Description, T1* Expression, T2* Expected, FAutomationTestBase& This)
 	{
-		This.TestEqual(*Description, reinterpret_cast<uint64>(Expression), reinterpret_cast<uint64>(Expected));
-		return Expression == Expected;
+		return This.TestEqual(*Description, reinterpret_cast<uint64>(Expression), reinterpret_cast<uint64>(Expected));
 	}
 
 	template<typename T1, typename T2>
 	inline bool TestNotEqual(const FString& Description, T1 Expression, T2 Expected, FAutomationTestBase& This)
 	{
-		This.TestNotEqual(*Description, Expression, Expected);
-		return Expression != Expected;
+		return This.TestNotEqual(*Description, Expression, Expected);
 	}
 
 	template<typename T1, typename T2>
 	inline bool TestNotEqual(const FString& Description, T1* Expression, T2* Expected, FAutomationTestBase& This)
 	{
-		This.TestNotEqual(*Description, reinterpret_cast<uint64>(Expression), reinterpret_cast<uint64>(Expected));
-		return Expression != Expected;
+		return This.TestNotEqual(*Description, reinterpret_cast<uint64>(Expression), reinterpret_cast<uint64>(Expected));
 	}
 }
 
