Removed all PawnActions-related tests from AITestSuite

diff --git a/Engine/Source/Developer/AITestSuite/Classes/Actions/TestPawnAction_CallFunction.h b/Engine/Source/Developer/AITestSuite/Classes/Actions/TestPawnAction_CallFunction.h
deleted file mode 100644
diff --git a/Engine/Source/Developer/AITestSuite/Classes/Actions/TestPawnAction_Log.h b/Engine/Source/Developer/AITestSuite/Classes/Actions/TestPawnAction_Log.h
deleted file mode 100644
diff --git a/Engine/Source/Developer/AITestSuite/Classes/MockAI.h b/Engine/Source/Developer/AITestSuite/Classes/MockAI.h
--- a/Engine/Source/Developer/AITestSuite/Classes/MockAI.h
+++ b/Engine/Source/Developer/AITestSuite/Classes/MockAI.h
@@ -12,7 +12,6 @@
 
 class UAIPerceptionComponent;
 class UBlackboardComponent;
-class UPawnActionsComponent;
 class UBrainComponent;
 class UMockAI;
 
@@ -48,9 +47,6 @@ class UMockAI : public UObject
 	UPROPERTY()
 	TObjectPtr<UAIPerceptionComponent> PerceptionComp = nullptr;
 
-	UPROPERTY()
-	TObjectPtr<UPawnActionsComponent> PawnActionComp = nullptr;
-	
 	template<typename TBrainClass>
 	void UseBrainComponent()
 	{
@@ -59,7 +55,6 @@ class UMockAI : public UObject
 
 	void UseBlackboardComponent();
 	void UsePerceptionComponent();
-	void UsePawnActionsComponent();
 
 	void SetEnableTicking(bool bShouldTick);
 
diff --git a/Engine/Source/Developer/AITestSuite/Private/Actions/TestPawnAction_CallFunction.cpp b/Engine/Source/Developer/AITestSuite/Private/Actions/TestPawnAction_CallFunction.cpp
deleted file mode 100644
diff --git a/Engine/Source/Developer/AITestSuite/Private/Actions/TestPawnAction_Log.cpp b/Engine/Source/Developer/AITestSuite/Private/Actions/TestPawnAction_Log.cpp
deleted file mode 100644
diff --git a/Engine/Source/Developer/AITestSuite/Private/MockAI/MockAI.cpp b/Engine/Source/Developer/AITestSuite/Private/MockAI/MockAI.cpp
--- a/Engine/Source/Developer/AITestSuite/Private/MockAI/MockAI.cpp
+++ b/Engine/Source/Developer/AITestSuite/Private/MockAI/MockAI.cpp
@@ -76,11 +76,6 @@ void UMockAI::UsePerceptionComponent()
 	PerceptionComp = NewObject<UAIPerceptionComponent>(Actor);
 }
 
-void UMockAI::UsePawnActionsComponent()
-{
-	PawnActionComp = NewObject<UPawnActionsComponent>(Actor);
-}
-
 void UMockAI::TickMe(float DeltaTime)
 {
 	if (BBComp)
@@ -97,11 +92,6 @@ void UMockAI::TickMe(float DeltaTime)
 	{
 		BrainComp->TickComponent(DeltaTime, ELevelTick::LEVELTICK_All, nullptr);
 	}
-	
-	if (PawnActionComp)
-	{
-		PawnActionComp->TickComponent(DeltaTime, ELevelTick::LEVELTICK_All, nullptr);
-	}
 }
 
 
diff --git a/Engine/Source/Developer/AITestSuite/Private/Tests/PawnActionsTest.cpp b/Engine/Source/Developer/AITestSuite/Private/Tests/PawnActionsTest.cpp
deleted file mode 100644
diff --git a/Engine/Source/Developer/AITestSuite/Private/Tests/ResourceIdTest.cpp b/Engine/Source/Developer/AITestSuite/Private/Tests/ResourceIdTest.cpp
--- a/Engine/Source/Developer/AITestSuite/Private/Tests/ResourceIdTest.cpp
+++ b/Engine/Source/Developer/AITestSuite/Private/Tests/ResourceIdTest.cpp
@@ -3,11 +3,9 @@
 #include "CoreMinimal.h"
 #include "AITypes.h"
 #include "AITestsCommon.h"
-#include "Actions/TestPawnAction_Log.h"
-#include "Actions/PawnActionsComponent.h"
 
 //----------------------------------------------------------------------//
-// 
+// FAITest_ResourceIDBasic
 //----------------------------------------------------------------------//
 struct FAITest_ResourceIDBasic : public FAITestBase
 {
@@ -24,7 +22,7 @@ struct FAITest_ResourceIDBasic : public FAITestBase
 IMPLEMENT_AI_INSTANT_TEST(FAITest_ResourceIDBasic, "System.AI.Resource ID.Basic operations")
 
 //----------------------------------------------------------------------//
-// 
+// FAITest_ResourceLock
 //----------------------------------------------------------------------//
 struct FAITest_ResourceLock : public FAITestBase
 {
@@ -74,7 +72,7 @@ struct FAITest_ResourceLock : public FAITestBase
 IMPLEMENT_AI_INSTANT_TEST(FAITest_ResourceLock, "System.AI.Resource ID.Resource locking")
 
 //----------------------------------------------------------------------//
-// 
+// FAITest_ResourceSet
 //----------------------------------------------------------------------//
 struct FAITest_ResourceSet : public FAITestBase
 {
@@ -131,66 +129,3 @@ struct FAITest_ResourceSet : public FAITestBase
 	}
 };
 IMPLEMENT_AI_INSTANT_TEST(FAITest_ResourceSet, "System.AI.Resource ID.Resource locking")
-
-//----------------------------------------------------------------------//
-// 
-//----------------------------------------------------------------------//
-struct FAITest_PawnActions_PausingActionsOfSameResource : public FAITest_SimpleComponentBasedTest<UPawnActionsComponent>
-{
-	virtual bool InstantTest() override
-	{
-		/*Logger.ExpectedValues.Add(ETestPawnActionMessage::Started);
-		Logger.ExpectedValues.Add(ETestPawnActionMessage::Paused);
-		Logger.ExpectedValues.Add(ETestPawnActionMessage::Started);*/
-
-		UWorld& World = GetWorld();
-		UTestPawnAction_Log* MoveAction = UTestPawnAction_Log::CreateAction(World, Logger);
-		MoveAction->GetRequiredResourcesSet() = FAIResourcesSet(FAIResources::Movement);
-		Component->PushAction(*MoveAction, EAIRequestPriority::Logic);
-
-		Component->TickComponent(FAITestHelpers::TickInterval, ELevelTick::LEVELTICK_All, nullptr);
-
-		UTestPawnAction_Log* AnotherMoveAction = UTestPawnAction_Log::CreateAction(World, Logger);
-		AnotherMoveAction->GetRequiredResourcesSet() = FAIResourcesSet(FAIResources::Movement);
-		Component->PushAction(*AnotherMoveAction, EAIRequestPriority::Logic);
-
-		Component->TickComponent(FAITestHelpers::TickInterval, ELevelTick::LEVELTICK_All, nullptr);
-
-		AITEST_TRUE("First MoveAction should get paused", MoveAction->IsPaused() == true);
-
-		return true;
-	}
-};
-IMPLEMENT_AI_INSTANT_TEST(FAITest_PawnActions_PausingActionsOfSameResource, "System.AI.Pawn Actions.Pausing actions of same resource")
-
-//----------------------------------------------------------------------//
-// 
-//----------------------------------------------------------------------//
-struct FAITest_PawnActions_NotPausingActionsOfDifferentResources : public FAITest_SimpleComponentBasedTest<UPawnActionsComponent>
-{
-	virtual bool InstantTest() override
-	{
-		/*Logger.ExpectedValues.Add(ETestPawnActionMessage::Started);
-		Logger.ExpectedValues.Add(ETestPawnActionMessage::Paused);
-		Logger.ExpectedValues.Add(ETestPawnActionMessage::Started);*/
-
-		UWorld& World = GetWorld();
-		UTestPawnAction_Log* MoveAction = UTestPawnAction_Log::CreateAction(World, Logger);
-		MoveAction->SetRequiredResourcesSet(FAIResourcesSet(FAIResources::Movement));
-		Component->PushAction(*MoveAction, EAIRequestPriority::Logic);
-
-		Component->TickComponent(FAITestHelpers::TickInterval, ELevelTick::LEVELTICK_All, nullptr);
-
-		UTestPawnAction_Log* PerceptionAction = UTestPawnAction_Log::CreateAction(World, Logger);
-		PerceptionAction->SetRequiredResourcesSet(FAIResourcesSet(FAIResources::Perception));
-		Component->PushAction(*PerceptionAction, EAIRequestPriority::Logic);
-
-		Component->TickComponent(FAITestHelpers::TickInterval, ELevelTick::LEVELTICK_All, nullptr);
-
-		// @todo test temporarily disabled
-		//AITEST_TRUE("First MoveAction should get paused", MoveAction->IsPaused() == false && PerceptionAction->IsPaused() == false);
-
-		return true;
-	}
-};
-IMPLEMENT_AI_INSTANT_TEST(FAITest_PawnActions_NotPausingActionsOfDifferentResources, "System.AI.Pawn Actions.Not pausing actions of different resources")
