Fixed a latent abort being resumed
Also fixed a VLog which was missing a parameter

diff --git a/Engine/Source/Developer/AITestSuite/Classes/BehaviorTree/TestBTTask_LatentWithFlags.h b/Engine/Source/Developer/AITestSuite/Classes/BehaviorTree/TestBTTask_LatentWithFlags.h
--- a/Engine/Source/Developer/AITestSuite/Classes/BehaviorTree/TestBTTask_LatentWithFlags.h
+++ b/Engine/Source/Developer/AITestSuite/Classes/BehaviorTree/TestBTTask_LatentWithFlags.h
@@ -14,6 +14,13 @@ struct FBTLatentTaskMemory
 	uint8 bIsAborting : 1;
 };
 
+UENUM()
+enum class EBTTestChangeFlagBehavior : uint8
+{
+	Set,
+	Toggle
+};
+
 UCLASS(meta = (HiddenNode))
 class UTestBTTask_LatentWithFlags : public UBTTaskNode
 {
@@ -55,6 +62,9 @@ public:
 	UPROPERTY()
 	FName KeyNameAbort = TEXT("Bool2");
 
+	UPROPERTY()
+	EBTTestChangeFlagBehavior ChangeFlagBehavior = EBTTestChangeFlagBehavior::Set;
+
 	UPROPERTY()
 	TEnumAsByte<EBTNodeResult::Type> LogResult = EBTNodeResult::Succeeded;
 
@@ -65,4 +75,6 @@ protected:
 	void LogExecution(class UBehaviorTreeComponent& OwnerComp, int32 LogNumber);
 
 	virtual void TickTask(class UBehaviorTreeComponent& OwnerComp, uint8* NodeMemory, float DeltaSeconds) override;
+
+	void ChangeFlag(UBehaviorTreeComponent& OwnerComp, FName FlagToChange) const;
 };
\ No newline at end of file
diff --git a/Engine/Source/Developer/AITestSuite/Private/BehaviorTree/TestBTTask_LatentWithFlags.cpp b/Engine/Source/Developer/AITestSuite/Private/BehaviorTree/TestBTTask_LatentWithFlags.cpp
--- a/Engine/Source/Developer/AITestSuite/Private/BehaviorTree/TestBTTask_LatentWithFlags.cpp
+++ b/Engine/Source/Developer/AITestSuite/Private/BehaviorTree/TestBTTask_LatentWithFlags.cpp
@@ -22,7 +22,7 @@ EBTNodeResult::Type UTestBTTask_LatentWithFlags::ExecuteTask(UBehaviorTreeCompon
 	LogExecution(OwnerComp, LogIndexExecuteStart);
 	if (ExecuteHalfTicks == 0)
 	{
-		OwnerComp.GetBlackboardComponent()->SetValueAsBool(KeyNameExecute, true);
+		ChangeFlag(OwnerComp, KeyNameExecute);
 		MyMemory->bFlagSet = true;
 
 		LogExecution(OwnerComp, LogIndexExecuteFinish);
@@ -43,7 +43,7 @@ EBTNodeResult::Type UTestBTTask_LatentWithFlags::AbortTask(UBehaviorTreeComponen
 	LogExecution(OwnerComp, LogIndexAbortStart);
 	if (AbortHalfTicks == 0)
 	{
-		OwnerComp.GetBlackboardComponent()->SetValueAsBool(KeyNameAbort, true);
+		ChangeFlag(OwnerComp, KeyNameAbort);
 		MyMemory->bFlagSet = true;
 
 		LogExecution(OwnerComp, LogIndexAbortFinish);
@@ -62,9 +62,7 @@ void UTestBTTask_LatentWithFlags::TickTask(UBehaviorTreeComponent& OwnerComp, ui
 	if (!MyMemory->bFlagSet && FAITestHelpers::FramesCounter() >= MyMemory->FlagFrameIdx)
 	{
 		MyMemory->bFlagSet = true;
-		OwnerComp.GetBlackboardComponent()->SetValueAsBool(
-			MyMemory->bIsAborting ? KeyNameAbort : KeyNameExecute,
-			true);
+		ChangeFlag(OwnerComp, MyMemory->bIsAborting ? KeyNameAbort : KeyNameExecute);
 	}
 
 	if (FAITestHelpers::FramesCounter() >= MyMemory->EndFrameIdx)
@@ -82,6 +80,19 @@ void UTestBTTask_LatentWithFlags::TickTask(UBehaviorTreeComponent& OwnerComp, ui
 	}
 }
 
+void UTestBTTask_LatentWithFlags::ChangeFlag(UBehaviorTreeComponent& OwnerComp, FName FlagToChange) const
+{
+	switch(ChangeFlagBehavior)
+	{
+		case EBTTestChangeFlagBehavior::Set:
+			OwnerComp.GetBlackboardComponent()->SetValueAsBool(FlagToChange, true);
+			break;
+		case EBTTestChangeFlagBehavior::Toggle:
+			OwnerComp.GetBlackboardComponent()->SetValueAsBool(FlagToChange, !OwnerComp.GetBlackboardComponent()->GetValueAsBool(FlagToChange));
+			break;
+	}
+}
+
 uint16 UTestBTTask_LatentWithFlags::GetInstanceMemorySize() const
 {
 	return sizeof(FBTLatentTaskMemory);
diff --git a/Engine/Source/Developer/AITestSuite/Private/Tests/BTTest.cpp b/Engine/Source/Developer/AITestSuite/Private/Tests/BTTest.cpp
--- a/Engine/Source/Developer/AITestSuite/Private/Tests/BTTest.cpp
+++ b/Engine/Source/Developer/AITestSuite/Private/Tests/BTTest.cpp
@@ -1740,6 +1740,49 @@ struct FAITest_BTAbortDuringLatentAbort4 : public FAITest_SimpleBT
 };
 IMPLEMENT_AI_LATENT_TEST(FAITest_BTAbortDuringLatentAbort4, "System.AI.Behavior Trees.Abort: during latent task abort (high pri 2)")
 
+struct FAITest_BTResumingDuringLatentAbort : public FAITest_SimpleBT
+{
+	FAITest_BTResumingDuringLatentAbort()
+	{
+		enum
+		{
+			LatentTaskStart = 0,
+			LatentTaskFinish,
+			LatentTaskAbortStart,
+			LatentTaskAbortFinish,
+			Task1Log,
+			Task2Log,
+		};
+
+
+		UBTCompositeNode& CompNode = FBTBuilder::AddSelector(*BTAsset); // 0
+		{
+			FBTBuilder::AddTask(CompNode, Task1Log, EBTNodeResult::Succeeded); // 2
+			{
+				FBTBuilder::WithDecoratorBlackboard(CompNode, EBasicKeyOperation::Set, EBTFlowAbortMode::Both, TEXT("Bool1")); // 1
+			}
+
+			FBTBuilder::AddTask(CompNode, Task2Log, EBTNodeResult::Succeeded); // 4
+			{
+				FBTBuilder::WithDecoratorBlackboard(CompNode, EBasicKeyOperation::Set, EBTFlowAbortMode::None, TEXT("Bool2")); // 3
+			}
+
+			FBTBuilder::AddTaskLatentFlags(CompNode, EBTNodeResult::Succeeded, // 5
+				/*ExecuteHalfTicks*/2, TEXT("Bool1"), LatentTaskStart, LatentTaskFinish,
+				/*AbortHalfTicks*/3, TEXT("Bool1"), LatentTaskAbortStart, LatentTaskAbortFinish, EBTTestChangeFlagBehavior::Toggle);
+			{
+				FBTBuilder::WithTaskServiceLog(CompNode, /*ActivationIndex*/-1, /*DeactivationIndex*/-1, /*TickIndex*/-1, /*TickBoolKeyName*/TEXT("Bool2")); // 6
+			}
+		}
+
+		ExpectedResult.Add(LatentTaskStart/*0*/);
+		ExpectedResult.Add(LatentTaskAbortStart/*2*/);
+		ExpectedResult.Add(LatentTaskAbortFinish/*3*/);
+		ExpectedResult.Add(Task2Log/*5*/);
+	}
+};
+IMPLEMENT_AI_LATENT_TEST(FAITest_BTResumingDuringLatentAbort, "System.AI.Behavior Trees.Resume: during latent task abort")
+
 
 struct FAITest_BTAbortDuringInstantAbort : public FAITest_SimpleBT
 {
diff --git a/Engine/Source/Developer/AITestSuite/Public/BTBuilder.h b/Engine/Source/Developer/AITestSuite/Public/BTBuilder.h
--- a/Engine/Source/Developer/AITestSuite/Public/BTBuilder.h
+++ b/Engine/Source/Developer/AITestSuite/Public/BTBuilder.h
@@ -206,7 +206,8 @@ struct FBTBuilder
 		int32 ExecuteHalfTicks, /** Num of ticks before 'execute start' and `set execute flag` and then the same num of ticks before `execute finish` */
 		FName ExecuteKeyName, int32 ExecuteLogStart, int32 ExecuteLogFinish,
 		int32 AbortHalfTicks = 0, /** Num of ticks before 'abort start' and `set abort flag` and then the same num of ticks before `abort finish` */
-		FName AbortKeyName = NAME_None, int32 AbortLogStart = 0, int32 AbortLogFinish = 0)
+		FName AbortKeyName = NAME_None, int32 AbortLogStart = 0, int32 AbortLogFinish = 0,
+		EBTTestChangeFlagBehavior ChangeFlagBehavior = EBTTestChangeFlagBehavior::Set)
 	{
 		UTestBTTask_LatentWithFlags* TaskNode = NewObject<UTestBTTask_LatentWithFlags>(ParentNode.GetTreeAsset());
 		TaskNode->ExecuteHalfTicks = ExecuteHalfTicks;
@@ -217,6 +218,7 @@ struct FBTBuilder
 		TaskNode->KeyNameAbort = AbortKeyName;
 		TaskNode->LogIndexAbortStart = AbortLogStart;
 		TaskNode->LogIndexAbortFinish = AbortLogFinish;
+		TaskNode->ChangeFlagBehavior = ChangeFlagBehavior;
 
 		const int32 ChildIdx = ParentNode.Children.AddZeroed(1);
 		ParentNode.Children[ChildIdx].ChildTask = TaskNode;
diff --git a/Engine/Source/Runtime/AIModule/Private/BehaviorTree/BTDecorator.cpp b/Engine/Source/Runtime/AIModule/Private/BehaviorTree/BTDecorator.cpp
--- a/Engine/Source/Runtime/AIModule/Private/BehaviorTree/BTDecorator.cpp
+++ b/Engine/Source/Runtime/AIModule/Private/BehaviorTree/BTDecorator.cpp
@@ -102,9 +102,10 @@ void UBTDecorator::ConditionalFlowAbort(UBehaviorTreeComponent& OwnerComp, EBTDe
 	const bool bPass = WrappedCanExecute(OwnerComp, NodeMemory);
 	const bool bAlwaysRequestWhenPassing = (RequestMode == EBTDecoratorAbortRequest::ConditionPassing);
 
-	UE_VLOG(OwnerComp.GetOwner(), LogBehaviorTree, Verbose, TEXT("%s, ConditionalFlowAbort(%s) pass:%d"),
+	UE_VLOG(OwnerComp.GetOwner(), LogBehaviorTree, Verbose, TEXT("%s, ConditionalFlowAbort(%s) pass:%s"),
 		*UBehaviorTreeTypes::DescribeNodeHelper(this),
-		bAlwaysRequestWhenPassing ? TEXT("always when passing") : TEXT("on change"));
+		bAlwaysRequestWhenPassing ? TEXT("always when passing") : TEXT("on change"), 
+		*LexToString(bPass));
 
 
 	if (!bPass)
diff --git a/Engine/Source/Runtime/AIModule/Private/BehaviorTree/BehaviorTreeComponent.cpp b/Engine/Source/Runtime/AIModule/Private/BehaviorTree/BehaviorTreeComponent.cpp
--- a/Engine/Source/Runtime/AIModule/Private/BehaviorTree/BehaviorTreeComponent.cpp
+++ b/Engine/Source/Runtime/AIModule/Private/BehaviorTree/BehaviorTreeComponent.cpp
@@ -690,6 +690,16 @@ void UBehaviorTreeComponent::DeactivateBranch(const UBTDecorator& RequestedBy)
 	else if (ensureMsgf(RequestedBy.GetParentNode() && RequestedBy.GetParentNode()->Children.IsValidIndex(RequestedBy.GetChildIndex()), 
 				TEXT("The decorator %s does not have a parent or is not a valid child."), *UBehaviorTreeTypes::DescribeNodeHelper(&RequestedBy)))
 	{
+		const bool bAbortPending = IsAbortPending();
+		if (bAbortPending)
+		{
+			// Branch that caused an abort and still waiting on a latent abort to complete might see its decorators changed again causing a call to DeactivateBranch.
+			// In this particular case we explicitly request the parent composite node to reevaluate without specifying a child index (i.e. RequestedByChildIndex = -1) so it will reevaluate all children.
+			UE_VLOG(GetOwner(), LogBehaviorTree, Verbose, TEXT("%s, Branch deactivation resulted in a reevaluation of the parent composite node because the abort was pending"), *UBehaviorTreeTypes::DescribeNodeHelper(&RequestedBy));
+			const int32 InstanceIdx = FindInstanceContainingNode(&RequestedBy);
+			RequestExecution(RequestedBy.GetParentNode(), InstanceIdx, &RequestedBy, -1, EBTNodeResult::Aborted);
+		}
+
 		UE_VLOG(GetOwner(), LogBehaviorTree, Verbose, TEXT("%s, Branch deactivation resulted in aux nodes unregistration"), *UBehaviorTreeTypes::DescribeNodeHelper(&RequestedBy));
 		if (const UBTCompositeNode* BranchRoot = RequestedBy.GetParentNode()->Children[RequestedBy.GetChildIndex()].ChildComposite)
 		{
