Added support to restart all nodes in RestartTree
UDN case 00532196

diff --git a/Engine/Source/Developer/AITestSuite/Classes/BehaviorTree/TestBTStopAction.h b/Engine/Source/Developer/AITestSuite/Classes/BehaviorTree/TestBTStopAction.h
--- a/Engine/Source/Developer/AITestSuite/Classes/BehaviorTree/TestBTStopAction.h
+++ b/Engine/Source/Developer/AITestSuite/Classes/BehaviorTree/TestBTStopAction.h
@@ -12,6 +12,7 @@ enum class EBTTestStopAction : uint8
 	UnInitialize,
 	Cleanup,
 	RestartTree,
+	FullRestartTree,
 	StartTree,
 };
 
diff --git a/Engine/Source/Developer/AITestSuite/Private/BehaviorTree/TestBTTask_BTStopAction.cpp b/Engine/Source/Developer/AITestSuite/Private/BehaviorTree/TestBTTask_BTStopAction.cpp
--- a/Engine/Source/Developer/AITestSuite/Private/BehaviorTree/TestBTTask_BTStopAction.cpp
+++ b/Engine/Source/Developer/AITestSuite/Private/BehaviorTree/TestBTTask_BTStopAction.cpp
@@ -70,7 +70,10 @@ void DoBTStopAction(UBehaviorTreeComponent& OwnerComp, const EBTTestStopAction S
 		OwnerComp.Cleanup();
 		break;
 	case EBTTestStopAction::RestartTree:
-		OwnerComp.RestartTree();
+		OwnerComp.RestartTree(EBTRestartMode::SkipReAddedNodes);
+		break;
+	case EBTTestStopAction::FullRestartTree:
+		OwnerComp.RestartTree(EBTRestartMode::ForceRestartAllNodes);
 		break;
 	case EBTTestStopAction::StartTree:
 		UBehaviorTree* NewBTAsset = NewObject<UBehaviorTree>(GetTransientPackage(), NAME_None, RF_NoFlags, OwnerComp.GetRootTree());
diff --git a/Engine/Source/Developer/AITestSuite/Private/Tests/BTTest.cpp b/Engine/Source/Developer/AITestSuite/Private/Tests/BTTest.cpp
--- a/Engine/Source/Developer/AITestSuite/Private/Tests/BTTest.cpp
+++ b/Engine/Source/Developer/AITestSuite/Private/Tests/BTTest.cpp
@@ -3845,4 +3845,103 @@ struct FAITest_BTCleanupDuringTaskServiceCeaseRelevant : public FAITest_SimpleBT
 };
 IMPLEMENT_AI_LATENT_TEST(FAITest_BTCleanupDuringTaskServiceCeaseRelevant, "System.AI.Behavior Trees.Cleanup: during task service cease relevant")
 
+//////////////////////////////////////////////////////////////////////////////////////////////////
+// Restart tests
+//////////////////////////////////////////////////////////////////////////////////////////////////
+struct FAITest_BTRestartDuringTaskExecute : public FAITest_SimpleBT
+{
+	FAITest_BTRestartDuringTaskExecute()
+	{
+		enum
+		{
+			ServiceActivate = 1,
+			ServiceDeactivate,
+			ServiceTick,
+			RestartTaskTaskExecute,
+			ServiceActivate2,
+			ServiceDeactivate2,
+			ServiceTick2,
+			Task2Execute
+		};
+
+		UBTCompositeNode& RootNode = FBTBuilder::AddSelector(*BTAsset);
+		UBTCompositeNode& CompNode = FBTBuilder::AddSelector(RootNode);
+		{
+			// Service
+			FBTBuilder::WithServiceLog(CompNode, ServiceActivate, ServiceDeactivate, ServiceTick);
+
+			UBTCompositeNode& Comp1Node = FBTBuilder::AddSelector(CompNode);
+			FBTBuilder::WithDecoratorBlackboard(CompNode, EBasicKeyOperation::NotSet, EBTFlowAbortMode::Both, TEXT("Bool1"));
+			{
+				FBTBuilder::WithServiceLog(Comp1Node, ServiceActivate2, ServiceDeactivate2, ServiceTick2, /*TickBoolKeyName*/NAME_None, /*bCallOnTickSearch*/false, /*BecomeRelevantBoolKeyName*/TEXT("Bool1"));
+
+				// RestartTask
+				FBTBuilder::AddTaskBTStopAction(Comp1Node, RestartTaskTaskExecute, EBTNodeResult::Succeeded, EBTTestTaskStopTiming::DuringExecute, EBTTestStopAction::RestartTree);
+			}
+			FBTBuilder::AddTask(CompNode, Task2Execute, EBTNodeResult::Succeeded);
+		}
+
+		ExpectedResult.Add(ServiceActivate/*1*/);
+		ExpectedResult.Add(ServiceActivate2/*5*/);
+		ExpectedResult.Add(ServiceTick/*3*/);
+		ExpectedResult.Add(ServiceTick2/*7*/);
+		ExpectedResult.Add(RestartTaskTaskExecute/*4*/);
+		ExpectedResult.Add(ServiceTick/*3*/);
+		ExpectedResult.Add(ServiceTick2/*7*/);
+		ExpectedResult.Add(ServiceDeactivate2/*6*/);
+		ExpectedResult.Add(Task2Execute/*8*/);
+		ExpectedResult.Add(ServiceTick/*3*/);
+		ExpectedResult.Add(ServiceDeactivate/*2*/);
+	}
+};
+IMPLEMENT_AI_LATENT_TEST(FAITest_BTRestartDuringTaskExecute, "System.AI.Behavior Trees.Restart: during task execute")
+
+struct FAITest_BTFullRestartDuringTaskExecute : public FAITest_SimpleBT
+{
+	FAITest_BTFullRestartDuringTaskExecute()
+	{
+		enum
+		{
+			ServiceActivate = 1,
+			ServiceDeactivate,
+			ServiceTick,
+			RestartTaskTaskExecute,
+			ServiceActivate2,
+			ServiceDeactivate2,
+			ServiceTick2,
+			Task2Execute
+		};
+
+		UBTCompositeNode& RootNode = FBTBuilder::AddSelector(*BTAsset);
+		UBTCompositeNode& CompNode = FBTBuilder::AddSelector(RootNode);
+		{
+			// Service
+			FBTBuilder::WithServiceLog(CompNode, ServiceActivate, ServiceDeactivate, ServiceTick);
+
+			UBTCompositeNode& Comp1Node = FBTBuilder::AddSelector(CompNode);
+			FBTBuilder::WithDecoratorBlackboard(CompNode, EBasicKeyOperation::NotSet, EBTFlowAbortMode::Both, TEXT("Bool1"));
+			{
+				FBTBuilder::WithServiceLog(Comp1Node, ServiceActivate2, ServiceDeactivate2, ServiceTick2, /*TickBoolKeyName*/NAME_None, /*bCallOnTickSearch*/false, /*BecomeRelevantBoolKeyName*/TEXT("Bool1"));
+
+				// RestartTask
+				FBTBuilder::AddTaskBTStopAction(Comp1Node, RestartTaskTaskExecute, EBTNodeResult::Succeeded, EBTTestTaskStopTiming::DuringExecute, EBTTestStopAction::FullRestartTree);
+			}
+			FBTBuilder::AddTask(CompNode, Task2Execute, EBTNodeResult::Succeeded);
+		}
+
+		ExpectedResult.Add(ServiceActivate/*1*/);
+		ExpectedResult.Add(ServiceActivate2/*5*/);
+		ExpectedResult.Add(ServiceTick/*3*/);
+		ExpectedResult.Add(ServiceTick2/*7*/);
+		ExpectedResult.Add(RestartTaskTaskExecute/*4*/);
+		ExpectedResult.Add(ServiceDeactivate/*2*/);
+		ExpectedResult.Add(ServiceDeactivate2/*6*/);
+		ExpectedResult.Add(ServiceActivate/*1*/);
+		ExpectedResult.Add(ServiceTick/*3*/);
+		ExpectedResult.Add(Task2Execute/*8*/);
+		ExpectedResult.Add(ServiceTick/*3*/);
+		ExpectedResult.Add(ServiceDeactivate/*2*/);
+	}
+};
+IMPLEMENT_AI_LATENT_TEST(FAITest_BTFullRestartDuringTaskExecute, "System.AI.Behavior Trees.Restart: full during task execute")
 #undef LOCTEXT_NAMESPACE
\ No newline at end of file
diff --git a/Engine/Source/Runtime/AIModule/Classes/BehaviorTree/BehaviorTreeComponent.h b/Engine/Source/Runtime/AIModule/Classes/BehaviorTree/BehaviorTreeComponent.h
--- a/Engine/Source/Runtime/AIModule/Classes/BehaviorTree/BehaviorTreeComponent.h
+++ b/Engine/Source/Runtime/AIModule/Classes/BehaviorTree/BehaviorTreeComponent.h
@@ -140,7 +140,7 @@ public:
 	void StopTree(EBTStopMode::Type StopMode = EBTStopMode::Safe);
 
 	/** restarts execution from root */
-	void RestartTree();
+	void RestartTree(EBTRestartMode RestartMode = EBTRestartMode::SkipReAddedNodes);
 
 	/** request execution change */
 	void RequestExecution(const UBTCompositeNode* RequestedOn, const int32 InstanceIdx, 
diff --git a/Engine/Source/Runtime/AIModule/Classes/BehaviorTree/BehaviorTreeTypes.h b/Engine/Source/Runtime/AIModule/Classes/BehaviorTree/BehaviorTreeTypes.h
--- a/Engine/Source/Runtime/AIModule/Classes/BehaviorTree/BehaviorTreeTypes.h
+++ b/Engine/Source/Runtime/AIModule/Classes/BehaviorTree/BehaviorTreeTypes.h
@@ -104,6 +104,12 @@ namespace EBTStopMode
 	};
 }
 
+enum class EBTRestartMode : uint8
+{
+	SkipReAddedNodes,
+	ForceRestartAllNodes,
+};
+
 namespace EBTMemoryInit
 {
 	enum Type
diff --git a/Engine/Source/Runtime/AIModule/Private/BehaviorTree/BehaviorTreeComponent.cpp b/Engine/Source/Runtime/AIModule/Private/BehaviorTree/BehaviorTreeComponent.cpp
--- a/Engine/Source/Runtime/AIModule/Private/BehaviorTree/BehaviorTreeComponent.cpp
+++ b/Engine/Source/Runtime/AIModule/Private/BehaviorTree/BehaviorTreeComponent.cpp
@@ -418,7 +418,7 @@ void UBehaviorTreeComponent::StopTree(EBTStopMode::Type StopMode)
 	bWaitingForLatentAborts = false;
 }
 
-void UBehaviorTreeComponent::RestartTree()
+void UBehaviorTreeComponent::RestartTree(EBTRestartMode RestartMode /*= EBTRestartMode::SkipReAddedNodes*/)
 {
 	UE_VLOG(GetOwner(), LogBehaviorTree, Log, TEXT("%s"), ANSI_TO_TCHAR(__FUNCTION__));
 	
@@ -440,8 +440,25 @@ void UBehaviorTreeComponent::RestartTree()
 	}
 	else if (InstanceStack.Num())
 	{
-		FBehaviorTreeInstance& TopInstance = InstanceStack[0];
-		RequestExecution(TopInstance.RootNode, 0, TopInstance.RootNode, -1, EBTNodeResult::Aborted);
+		switch(RestartMode)
+		{
+			case EBTRestartMode::SkipReAddedNodes:
+			{
+				FBehaviorTreeInstance& TopInstance = InstanceStack[0];
+				RequestExecution(TopInstance.RootNode, 0, TopInstance.RootNode, -1, EBTNodeResult::Aborted);
+				break;
+			}
+			case EBTRestartMode::ForceRestartAllNodes:
+			{
+				StopTree(EBTStopMode::Safe);
+
+				TreeStartInfo.Asset = GetRootTree();
+				TreeStartInfo.ExecuteMode = bLoopExecution ? EBTExecutionMode::Looped : EBTExecutionMode::SingleRun;
+				TreeStartInfo.bPendingInitialize = true;
+				ProcessPendingInitialize();
+				break;
+			}
+		}
 	}
 }
 
