Better description and information of the my previous RestartTree change.

diff --git a/Engine/Source/Developer/AITestSuite/Classes/BehaviorTree/TestBTStopAction.h b/Engine/Source/Developer/AITestSuite/Classes/BehaviorTree/TestBTStopAction.h
--- a/Engine/Source/Developer/AITestSuite/Classes/BehaviorTree/TestBTStopAction.h
+++ b/Engine/Source/Developer/AITestSuite/Classes/BehaviorTree/TestBTStopAction.h
@@ -11,8 +11,8 @@ enum class EBTTestStopAction : uint8
 	StopTree,
 	UnInitialize,
 	Cleanup,
-	RestartTree,
-	FullRestartTree,
+	Restart_ForceReevaluateRootNode,
+	Restart_Complete,
 	StartTree,
 };
 
diff --git a/Engine/Source/Developer/AITestSuite/Private/BehaviorTree/TestBTTask_BTStopAction.cpp b/Engine/Source/Developer/AITestSuite/Private/BehaviorTree/TestBTTask_BTStopAction.cpp
--- a/Engine/Source/Developer/AITestSuite/Private/BehaviorTree/TestBTTask_BTStopAction.cpp
+++ b/Engine/Source/Developer/AITestSuite/Private/BehaviorTree/TestBTTask_BTStopAction.cpp
@@ -69,11 +69,11 @@ void DoBTStopAction(UBehaviorTreeComponent& OwnerComp, const EBTTestStopAction S
 	case EBTTestStopAction::Cleanup:
 		OwnerComp.Cleanup();
 		break;
-	case EBTTestStopAction::RestartTree:
-		OwnerComp.RestartTree(EBTRestartMode::SkipReAddedNodes);
+	case EBTTestStopAction::Restart_ForceReevaluateRootNode:
+		OwnerComp.RestartTree(EBTRestartMode::ForceReevaluateRootNode);
 		break;
-	case EBTTestStopAction::FullRestartTree:
-		OwnerComp.RestartTree(EBTRestartMode::ForceRestartAllNodes);
+	case EBTTestStopAction::Restart_Complete:
+		OwnerComp.RestartTree(EBTRestartMode::CompleteRestart);
 		break;
 	case EBTTestStopAction::StartTree:
 		UBehaviorTree* NewBTAsset = NewObject<UBehaviorTree>(GetTransientPackage(), NAME_None, RF_NoFlags, OwnerComp.GetRootTree());
diff --git a/Engine/Source/Developer/AITestSuite/Private/Tests/BTTest.cpp b/Engine/Source/Developer/AITestSuite/Private/Tests/BTTest.cpp
--- a/Engine/Source/Developer/AITestSuite/Private/Tests/BTTest.cpp
+++ b/Engine/Source/Developer/AITestSuite/Private/Tests/BTTest.cpp
@@ -3848,9 +3848,9 @@ IMPLEMENT_AI_LATENT_TEST(FAITest_BTCleanupDuringTaskServiceCeaseRelevant, "Syste
 //////////////////////////////////////////////////////////////////////////////////////////////////
 // Restart tests
 //////////////////////////////////////////////////////////////////////////////////////////////////
-struct FAITest_BTRestartDuringTaskExecute : public FAITest_SimpleBT
+struct FAITest_BTDefaultRestartDuringTaskExecute : public FAITest_SimpleBT
 {
-	FAITest_BTRestartDuringTaskExecute()
+	FAITest_BTDefaultRestartDuringTaskExecute()
 	{
 		enum
 		{
@@ -3876,7 +3876,7 @@ struct FAITest_BTRestartDuringTaskExecute : public FAITest_SimpleBT
 				FBTBuilder::WithServiceLog(Comp1Node, ServiceActivate2, ServiceDeactivate2, ServiceTick2, /*TickBoolKeyName*/NAME_None, /*bCallOnTickSearch*/false, /*BecomeRelevantBoolKeyName*/TEXT("Bool1"));
 
 				// RestartTask
-				FBTBuilder::AddTaskBTStopAction(Comp1Node, RestartTaskTaskExecute, EBTNodeResult::Succeeded, EBTTestTaskStopTiming::DuringExecute, EBTTestStopAction::RestartTree);
+				FBTBuilder::AddTaskBTStopAction(Comp1Node, RestartTaskTaskExecute, EBTNodeResult::Succeeded, EBTTestTaskStopTiming::DuringExecute, EBTTestStopAction::Restart_ForceReevaluateRootNode);
 			}
 			FBTBuilder::AddTask(CompNode, Task2Execute, EBTNodeResult::Succeeded);
 		}
@@ -3894,11 +3894,11 @@ struct FAITest_BTRestartDuringTaskExecute : public FAITest_SimpleBT
 		ExpectedResult.Add(ServiceDeactivate/*2*/);
 	}
 };
-IMPLEMENT_AI_LATENT_TEST(FAITest_BTRestartDuringTaskExecute, "System.AI.Behavior Trees.Restart: during task execute")
+IMPLEMENT_AI_LATENT_TEST(FAITest_BTDefaultRestartDuringTaskExecute, "System.AI.Behavior Trees.Restart: default during task execute")
 
-struct FAITest_BTFullRestartDuringTaskExecute : public FAITest_SimpleBT
+struct FAITest_BTCompleteRestartDuringTaskExecute : public FAITest_SimpleBT
 {
-	FAITest_BTFullRestartDuringTaskExecute()
+	FAITest_BTCompleteRestartDuringTaskExecute()
 	{
 		enum
 		{
@@ -3924,7 +3924,7 @@ struct FAITest_BTFullRestartDuringTaskExecute : public FAITest_SimpleBT
 				FBTBuilder::WithServiceLog(Comp1Node, ServiceActivate2, ServiceDeactivate2, ServiceTick2, /*TickBoolKeyName*/NAME_None, /*bCallOnTickSearch*/false, /*BecomeRelevantBoolKeyName*/TEXT("Bool1"));
 
 				// RestartTask
-				FBTBuilder::AddTaskBTStopAction(Comp1Node, RestartTaskTaskExecute, EBTNodeResult::Succeeded, EBTTestTaskStopTiming::DuringExecute, EBTTestStopAction::FullRestartTree);
+				FBTBuilder::AddTaskBTStopAction(Comp1Node, RestartTaskTaskExecute, EBTNodeResult::Succeeded, EBTTestTaskStopTiming::DuringExecute, EBTTestStopAction::Restart_Complete);
 			}
 			FBTBuilder::AddTask(CompNode, Task2Execute, EBTNodeResult::Succeeded);
 		}
@@ -3943,5 +3943,5 @@ struct FAITest_BTFullRestartDuringTaskExecute : public FAITest_SimpleBT
 		ExpectedResult.Add(ServiceDeactivate/*2*/);
 	}
 };
-IMPLEMENT_AI_LATENT_TEST(FAITest_BTFullRestartDuringTaskExecute, "System.AI.Behavior Trees.Restart: full during task execute")
+IMPLEMENT_AI_LATENT_TEST(FAITest_BTCompleteRestartDuringTaskExecute, "System.AI.Behavior Trees.Restart: complete during task execute")
 #undef LOCTEXT_NAMESPACE
\ No newline at end of file
diff --git a/Engine/Source/Runtime/AIModule/Classes/BehaviorTree/BehaviorTreeComponent.h b/Engine/Source/Runtime/AIModule/Classes/BehaviorTree/BehaviorTreeComponent.h
--- a/Engine/Source/Runtime/AIModule/Classes/BehaviorTree/BehaviorTreeComponent.h
+++ b/Engine/Source/Runtime/AIModule/Classes/BehaviorTree/BehaviorTreeComponent.h
@@ -139,8 +139,11 @@ public:
 	/** stops execution */
 	void StopTree(EBTStopMode::Type StopMode = EBTStopMode::Safe);
 
-	/** restarts execution from root */
-	void RestartTree(EBTRestartMode RestartMode = EBTRestartMode::SkipReAddedNodes);
+	/** restarts execution from root
+	 * @param RestartMode to force the reevaluation of the root node, which could skip active nodes that are removed and then readded (default)
+	 *        or a complete restart from scratch of the tree (equivalent of StopTree then StartTree)
+	 */
+	void RestartTree(EBTRestartMode RestartMode = EBTRestartMode::ForceReevaluateRootNode);
 
 	/** request execution change */
 	void RequestExecution(const UBTCompositeNode* RequestedOn, const int32 InstanceIdx, 
diff --git a/Engine/Source/Runtime/AIModule/Classes/BehaviorTree/BehaviorTreeTypes.h b/Engine/Source/Runtime/AIModule/Classes/BehaviorTree/BehaviorTreeTypes.h
--- a/Engine/Source/Runtime/AIModule/Classes/BehaviorTree/BehaviorTreeTypes.h
+++ b/Engine/Source/Runtime/AIModule/Classes/BehaviorTree/BehaviorTreeTypes.h
@@ -106,8 +106,8 @@ namespace EBTStopMode
 
 enum class EBTRestartMode : uint8
 {
-	SkipReAddedNodes,
-	ForceRestartAllNodes,
+	ForceReevaluateRootNode, // (Default) will just request a new execution on the root node and any active nodes that gets re-added will not get CeaseRelevant/BecomeRelevant notification
+	CompleteRestart, // essentially equivalent to calling StopTree and then StartTree. Every active node is going to be removed and the execution started from the root
 };
 
 namespace EBTMemoryInit
diff --git a/Engine/Source/Runtime/AIModule/Private/BehaviorTree/BehaviorTreeComponent.cpp b/Engine/Source/Runtime/AIModule/Private/BehaviorTree/BehaviorTreeComponent.cpp
--- a/Engine/Source/Runtime/AIModule/Private/BehaviorTree/BehaviorTreeComponent.cpp
+++ b/Engine/Source/Runtime/AIModule/Private/BehaviorTree/BehaviorTreeComponent.cpp
@@ -442,20 +442,20 @@ void UBehaviorTreeComponent::RestartTree(EBTRestartMode RestartMode /*= EBTResta
 	{
 		switch(RestartMode)
 		{
-			case EBTRestartMode::SkipReAddedNodes:
+			case EBTRestartMode::ForceReevaluateRootNode:
 			{
 				FBehaviorTreeInstance& TopInstance = InstanceStack[0];
 				RequestExecution(TopInstance.RootNode, 0, TopInstance.RootNode, -1, EBTNodeResult::Aborted);
 				break;
 			}
-			case EBTRestartMode::ForceRestartAllNodes:
+			case EBTRestartMode::CompleteRestart:
 			{
-				StopTree(EBTStopMode::Safe);
-
 				TreeStartInfo.Asset = GetRootTree();
 				TreeStartInfo.ExecuteMode = bLoopExecution ? EBTExecutionMode::Looped : EBTExecutionMode::SingleRun;
 				TreeStartInfo.bPendingInitialize = true;
 				ProcessPendingInitialize();
+
+				StopTree(EBTStopMode::Safe);
 				break;
 			}
 		}
