Turned off by default the code which applys aux node changes for failed searches as it caused further minor issues. I added a BT test to catch this case which will fire if the code is reenabled.

diff --git a/Engine/Source/Developer/AITestSuite/Classes/BehaviorTree/TestBTService_Log.h b/Engine/Source/Developer/AITestSuite/Classes/BehaviorTree/TestBTService_Log.h
--- a/Engine/Source/Developer/AITestSuite/Classes/BehaviorTree/TestBTService_Log.h
+++ b/Engine/Source/Developer/AITestSuite/Classes/BehaviorTree/TestBTService_Log.h
@@ -29,6 +29,12 @@ class UTestBTService_Log : public UBTService
 	UPROPERTY()
 	int32 LogTick;
 
+	UPROPERTY()
+	int32 TicksDelaySetKeyNameTick;
+
+	UPROPERTY()
+	int32 NumTicks;
+
 	UPROPERTY()
 	bool bToggleValue;
 
diff --git a/Engine/Source/Developer/AITestSuite/Private/BehaviorTree/TestBTService_Log.cpp b/Engine/Source/Developer/AITestSuite/Private/BehaviorTree/TestBTService_Log.cpp
--- a/Engine/Source/Developer/AITestSuite/Private/BehaviorTree/TestBTService_Log.cpp
+++ b/Engine/Source/Developer/AITestSuite/Private/BehaviorTree/TestBTService_Log.cpp
@@ -59,7 +59,7 @@ void UTestBTService_Log::TickNode(UBehaviorTreeComponent& OwnerComp, uint8* Node
 {
 	Super::TickNode(OwnerComp, NodeMemory, DeltaSeconds);
 
-	if (KeyNameTick != NAME_None)
+	if (KeyNameTick != NAME_None && NumTicks >= TicksDelaySetKeyNameTick)
 	{
 		OwnerComp.GetBlackboardComponent()->SetValueAsBool(KeyNameTick, bToggleValue ? !OwnerComp.GetBlackboardComponent()->GetValueAsBool(KeyNameTick) : true);
 	}
@@ -68,6 +68,7 @@ void UTestBTService_Log::TickNode(UBehaviorTreeComponent& OwnerComp, uint8* Node
 	{
 		UMockAI_BT::ExecutionLog.Add(LogTick);
 	}
+	++NumTicks;
 }
 
 void UTestBTService_Log::SetFlagOnTick(FName InKeyNameTick, bool bInCallTickOnSearchStart /* = false */)
diff --git a/Engine/Source/Developer/AITestSuite/Private/Tests/BTTest.cpp b/Engine/Source/Developer/AITestSuite/Private/Tests/BTTest.cpp
--- a/Engine/Source/Developer/AITestSuite/Private/Tests/BTTest.cpp
+++ b/Engine/Source/Developer/AITestSuite/Private/Tests/BTTest.cpp
@@ -3945,37 +3945,84 @@ struct FAITest_BTCompleteRestartDuringTaskExecute : public FAITest_SimpleBT
 };
 IMPLEMENT_AI_LATENT_TEST(FAITest_BTCompleteRestartDuringTaskExecute, "System.AI.Behavior Trees.Restart: complete during task execute")
 
-struct FAITest_BTTestOberserverAbortLowerPriorityTwoDeep: public FAITest_SimpleBT
-{
-	FAITest_BTTestOberserverAbortLowerPriorityTwoDeep()
+// This test is temporarily turned off until we enable CVarApplyAuxNodesFromFailedSearches again
+//struct FAITest_BTTestOberserverAbortLowerPriorityTwoDeep: public FAITest_SimpleBT
+//{
+//	FAITest_BTTestOberserverAbortLowerPriorityTwoDeep()
+//	{
+//		enum
+//		{
+//			MainTaskExecute = 1,
+//			InteruptingTaskExecute,
+//		};
+//
+//		UBTCompositeNode& RootNode = FBTBuilder::AddSelector(*BTAsset);
+//		UBTCompositeNode& CompNode = FBTBuilder::AddSelector(RootNode);
+//
+//		{
+//			UBTCompositeNode& Comp1Node = FBTBuilder::AddSelector(CompNode);
+//			FBTBuilder::WithDecoratorBlackboard(CompNode, EArithmeticKeyOperation::Equal, 1, EBTFlowAbortMode::LowerPriority, EBTBlackboardRestart::ValueChange, TEXT("Int"));
+//
+//			UBTCompositeNode& Comp2Node = FBTBuilder::AddSelector(Comp1Node);
+//			FBTBuilder::WithDecoratorBlackboard(Comp1Node, EArithmeticKeyOperation::Equal, 1, EBTFlowAbortMode::LowerPriority, EBTBlackboardRestart::ValueChange, TEXT("Int2"));
+//
+//			FBTBuilder::AddTask(Comp2Node, InteruptingTaskExecute, EBTNodeResult::Succeeded);
+//		}
+//
+//		{
+//			FBTBuilder::AddTaskValuesChangedWithLogs(/*ParentNode*/ CompNode, /*LogIndex*/ MainTaskExecute, /*NodeResult*/ EBTNodeResult::Succeeded, /*Value1*/1, /*Value2*/ 1, /*IntKeyName2*/ TEXT("Int"), /* IntKeyName2 */ TEXT("Int2"), /*ExecutionTicks1*/ 5, /*ExecutionTicks2*/ 5);
+//		}
+//
+//		ExpectedResult.Add(MainTaskExecute);
+//		ExpectedResult.Add(InteruptingTaskExecute);
+//	}
+//};
+//IMPLEMENT_AI_LATENT_TEST(FAITest_BTTestOberserverAbortLowerPriorityTwoDeep, "System.AI.Behavior Trees.OberserverAbortLowerPriorityTwoDeep: Two nodes in same branch with observer aborts lower priority")
+
+struct FAITest_BTTestTickAuxNodesDeactivateNodes : public FAITest_SimpleBT
+{
+	FAITest_BTTestTickAuxNodesDeactivateNodes()
 	{
 		enum
 		{
 			MainTaskExecute = 1,
-			InteruptingTaskExecute,
+			StartAbortMainTask,
+			SecondaryTask,
+			IncorrectTaskExecute,
 		};
 
 		UBTCompositeNode& RootNode = FBTBuilder::AddSelector(*BTAsset);
-		UBTCompositeNode& CompNode = FBTBuilder::AddSelector(RootNode);
+		FBTBuilder::WithServiceLog(/*ParentNode*/ RootNode, /*ActivationIndex*/ INDEX_NONE, /*DeactivationIndex*/ INDEX_NONE, /*TickIndex*/ INDEX_NONE, /*TickBoolKeyName*/ TEXT("Bool1"), /*bCallTickOnSearchStart*/ false, /*BecomeRelevantBoolKeyName*/ NAME_None, /*BecomeRelevantBoolKeyName*/ NAME_None, /*bToggleValue*/ false, /*TicksDelaySetKeyNameTick*/ 4);
+		FBTBuilder::WithServiceLog(/*ParentNode*/ RootNode, /*ActivationIndex*/ INDEX_NONE, /*DeactivationIndex*/ INDEX_NONE, /*TickIndex*/ INDEX_NONE, /*TickBoolKeyName*/ TEXT("Bool3"), /*bCallTickOnSearchStart*/ false, /*BecomeRelevantBoolKeyName*/ NAME_None, /*BecomeRelevantBoolKeyName*/ NAME_None, /*bToggleValue*/ false, /*TicksDelaySetKeyNameTick*/ 2);
 
 		{
-			UBTCompositeNode& Comp1Node = FBTBuilder::AddSelector(CompNode);
-			FBTBuilder::WithDecoratorBlackboard(CompNode, EArithmeticKeyOperation::Equal, 1, EBTFlowAbortMode::LowerPriority, EBTBlackboardRestart::ValueChange, TEXT("Int"));
+			UBTCompositeNode& CompNode = FBTBuilder::AddSelector(RootNode);
+			FBTBuilder::WithDecoratorBlackboard(/*ParentNode*/ RootNode, /*Condition*/ EBasicKeyOperation::NotSet, /*Observer*/ EBTFlowAbortMode::LowerPriority, /*BoolKeyName*/  TEXT("Bool1"));
+			{
+				UBTCompositeNode& CompNode2 = FBTBuilder::AddSelector(CompNode);
+				FBTBuilder::WithDecoratorBlackboard(/*ParentNode*/ RootNode, /*Condition*/ EBasicKeyOperation::Set, /*Observer*/ EBTFlowAbortMode::LowerPriority, /*BoolKeyName*/ TEXT("Bool2"));
+				{
+					FBTBuilder::AddTaskLogFinish(/*ParentNode*/ CompNode2, /*LogIndex*/ IncorrectTaskExecute, /*FinishIndex*/ INDEX_NONE, EBTNodeResult::Succeeded, /*ExecutionTick*/ 0);
+				}
+			}
+		}
 
-			UBTCompositeNode& Comp2Node = FBTBuilder::AddSelector(Comp1Node);
-			FBTBuilder::WithDecoratorBlackboard(Comp1Node, EArithmeticKeyOperation::Equal, 1, EBTFlowAbortMode::LowerPriority, EBTBlackboardRestart::ValueChange, TEXT("Int2"));
+		{
+			UBTCompositeNode& CompNode3 = FBTBuilder::AddSelector(RootNode);
+			FBTBuilder::WithDecoratorBlackboard(/*ParentNode*/ RootNode, /*Condition*/ EBasicKeyOperation::Set, /*Observer*/ EBTFlowAbortMode::LowerPriority, /*BoolKeyName*/ TEXT("Bool3"));
 
-			FBTBuilder::AddTask(Comp2Node, InteruptingTaskExecute, EBTNodeResult::Succeeded);
+			FBTBuilder::AddTaskLogFinish(/*ParentNode*/ CompNode3, /*LogIndex*/ SecondaryTask, /*FinishIndex*/ INDEX_NONE, EBTNodeResult::Succeeded, /*ExecutionTick*/ 4);
 		}
 
 		{
-			FBTBuilder::AddTaskValuesChangedWithLogs(/*ParentNode*/ CompNode, /*LogIndex*/ MainTaskExecute, /*NodeResult*/ EBTNodeResult::Succeeded, /*Value1*/1, /*Value2*/ 1, /*IntKeyName2*/ TEXT("Int"), /* IntKeyName2 */ TEXT("Int2"), /*ExecutionTicks1*/ 5, /*ExecutionTicks2*/ 5);
+			FBTBuilder::AddTaskLatentFlags(RootNode, EBTNodeResult::Succeeded, 3, FName(), MainTaskExecute, INDEX_NONE, 3, NAME_None, 2, INDEX_NONE);
 		}
 
 		ExpectedResult.Add(MainTaskExecute);
-		ExpectedResult.Add(InteruptingTaskExecute);
+		ExpectedResult.Add(StartAbortMainTask);
+		ExpectedResult.Add(SecondaryTask);
 	}
 };
-IMPLEMENT_AI_LATENT_TEST(FAITest_BTTestOberserverAbortLowerPriorityTwoDeep, "System.AI.Behavior Trees.OberserverAbortLowerPriorityTwoDeep: Two nodes in same branch with observer aborts lower priority")
+IMPLEMENT_AI_LATENT_TEST(FAITest_BTTestTickAuxNodesDeactivateNodes, "System.AI.Behavior Trees.AITest_TickAuxNodesDeactivateNodes: Ticking Aux Nodes Deactivates Nodes")
 
 #undef LOCTEXT_NAMESPACE
\ No newline at end of file
diff --git a/Engine/Source/Developer/AITestSuite/Public/BTBuilder.h b/Engine/Source/Developer/AITestSuite/Public/BTBuilder.h
--- a/Engine/Source/Developer/AITestSuite/Public/BTBuilder.h
+++ b/Engine/Source/Developer/AITestSuite/Public/BTBuilder.h
@@ -359,7 +359,7 @@ struct FBTBuilder
 		return *ServiceOb;
 	}
 
-	static void WithServiceLog(UBTCompositeNode& ParentNode, int32 ActivationIndex, int32 DeactivationIndex, int32 TickIndex = INDEX_NONE, FName TickBoolKeyName = NAME_None, bool bCallTickOnSearchStart = false, FName BecomeRelevantBoolKeyName = NAME_None, FName CeaseRelevantBoolKeyName = NAME_None, bool bToggleValue = false)
+	static void WithServiceLog(UBTCompositeNode& ParentNode, int32 ActivationIndex, int32 DeactivationIndex, int32 TickIndex = INDEX_NONE, FName TickBoolKeyName = NAME_None, bool bCallTickOnSearchStart = false, FName BecomeRelevantBoolKeyName = NAME_None, FName CeaseRelevantBoolKeyName = NAME_None, bool bToggleValue = false, int32 TicksDelaySetKeyNameTick = 0)
 	{
 		UTestBTService_Log& LogService = WithService<UTestBTService_Log>(ParentNode);
 		LogService.LogActivation = ActivationIndex;
@@ -369,6 +369,7 @@ struct FBTBuilder
 		LogService.KeyNameBecomeRelevant = BecomeRelevantBoolKeyName;
 		LogService.KeyNameCeaseRelevant = CeaseRelevantBoolKeyName;
 		LogService.bToggleValue = bToggleValue;
+		LogService.TicksDelaySetKeyNameTick = TicksDelaySetKeyNameTick;
 	}
 
 	static void WithServiceBTStopAction(UBTCompositeNode& ParentNode, int32 LogIndex, EBTTestServiceStopTiming StopTiming, EBTTestStopAction StopAction)
diff --git a/Engine/Source/Runtime/AIModule/Private/BehaviorTree/BehaviorTreeComponent.cpp b/Engine/Source/Runtime/AIModule/Private/BehaviorTree/BehaviorTreeComponent.cpp
--- a/Engine/Source/Runtime/AIModule/Private/BehaviorTree/BehaviorTreeComponent.cpp
+++ b/Engine/Source/Runtime/AIModule/Private/BehaviorTree/BehaviorTreeComponent.cpp
@@ -33,7 +33,16 @@ double UBehaviorTreeComponent::FrameSearchTime = 0.;
 int32 UBehaviorTreeComponent::NumSearchTimeCalls = 0;
 #endif
 
-static TAutoConsoleVariable<bool> CVarAddAuxNodesFromFailedSearches(TEXT("BehaviorTree.AddAuxNodesFromFailedSearches"), 1, TEXT("Add Aux Nodes From Failed Searches"));
+namespace UE::BehaviorTreeCVars
+{
+	// Note this is defaulted to off for now as it caused a further bug, there is a BT unit test that will fire if this
+	// code is re-enabled, once that is addressed this can be re-enabled by default.
+	static bool bApplyAuxNodesFromFailedSearches = false;
+	static FAutoConsoleVariableRef CVarApplyAuxNodesFromFailedSearches(
+		TEXT("BehaviorTree.ApplyAuxNodesFromFailedSearches"),
+		bApplyAuxNodesFromFailedSearches,
+		TEXT("Apply Aux Nodes From Failed Searches"));
+}
 
 //----------------------------------------------------------------------//
 // UBehaviorTreeComponent
@@ -1638,7 +1647,7 @@ void UBehaviorTreeComponent::ApplySearchData(UBTNode* NewActiveNode)
 
 void UBehaviorTreeComponent::ApplyDiscardedSearch()
 {
-	if (CVarAddAuxNodesFromFailedSearches.GetValueOnGameThread())
+	if (UE::BehaviorTreeCVars::bApplyAuxNodesFromFailedSearches)
 	{
 		// Apply aux nodes from last search with the currently active node as the 'NewNode'
 		int32 NewNodeExecutionIdx = 0;
