[BehaviorTree] prevent ForceSuccess decorator from affecting node results from unrelated nodes. Should only modify result if associated task or composite node fails.

diff --git a/Engine/Source/Developer/AITestSuite/Private/Tests/BTTest.cpp b/Engine/Source/Developer/AITestSuite/Private/Tests/BTTest.cpp
--- a/Engine/Source/Developer/AITestSuite/Private/Tests/BTTest.cpp
+++ b/Engine/Source/Developer/AITestSuite/Private/Tests/BTTest.cpp
@@ -139,6 +139,9 @@ struct FAITest_BTBasicSequence : public FAITest_SimpleBT
 			FBTBuilder::AddTask(CompNode, 1, EBTNodeResult::Failed);
 			{
 				FBTBuilder::WithDecorator<UTestBTDecorator_CantExecute>(CompNode);
+
+				// @todo BT: We should reconsider allowing decorator to change the activation result.
+				// This looks more like a debugger feature than a desired behavior. 
 				FBTBuilder::WithDecorator<UBTDecorator_ForceSuccess>(CompNode);
 			}
 
@@ -231,6 +234,146 @@ struct FAITest_BTDecoratorBlueprintNoConditionImpl : public FAITest_SimpleBT
 };
 IMPLEMENT_AI_LATENT_TEST(FAITest_BTDecoratorBlueprintNoConditionImpl, "System.AI.Behavior Trees.Decorator: blueprint no condition impl")
 
+struct FAITest_BTDecoratorForceSuccessTaskFailed : public FAITest_SimpleBT
+{
+	FAITest_BTDecoratorForceSuccessTaskFailed()
+	{
+		enum
+		{
+			Task1Execute,
+			Task2Execute,
+			Task3Execute,
+			Task4Execute,
+		};
+
+		UBTCompositeNode& MainSequence = FBTBuilder::AddSequence(*BTAsset);
+		{
+			FBTBuilder::AddTask(MainSequence, Task1Execute, EBTNodeResult::Succeeded);
+
+			UBTCompositeNode& InnerSequence = FBTBuilder::AddSequence(MainSequence);
+			{
+				FBTBuilder::AddTask(InnerSequence, Task2Execute, EBTNodeResult::Failed);
+				FBTBuilder::WithDecorator<UBTDecorator_ForceSuccess>(InnerSequence);
+
+				FBTBuilder::AddTask(InnerSequence, Task3Execute, EBTNodeResult::Succeeded);
+			}
+
+			FBTBuilder::AddTask(MainSequence, Task4Execute, EBTNodeResult::Succeeded);
+		}
+
+		ExpectedResult.Add(Task1Execute/*0*/);
+		ExpectedResult.Add(Task2Execute/*1*/);
+		ExpectedResult.Add(Task3Execute/*2*/);
+		ExpectedResult.Add(Task4Execute/*3*/);
+	}
+};
+IMPLEMENT_AI_LATENT_TEST(FAITest_BTDecoratorForceSuccessTaskFailed, "System.AI.Behavior Trees.Decorator: force success task failed self")
+
+struct FAITest_BTDecoratorForceSuccessTaskFailedParentComposite : public FAITest_SimpleBT
+{
+	FAITest_BTDecoratorForceSuccessTaskFailedParentComposite()
+	{
+		enum
+		{
+			Task1Execute,
+			Task2Execute,
+			Task3Execute,
+			Task4Execute,
+		};
+
+		UBTCompositeNode& MainSequence = FBTBuilder::AddSequence(*BTAsset);
+		{
+			FBTBuilder::AddTask(MainSequence, Task1Execute, EBTNodeResult::Succeeded);
+
+			UBTCompositeNode& InnerSequence = FBTBuilder::AddSequence(MainSequence);
+			{
+				FBTBuilder::WithDecoratorBlackboard(MainSequence, EArithmeticKeyOperation::NotEqual, /*IntValue*/1, EBTFlowAbortMode::Self, EBTBlackboardRestart::ResultChange);
+
+				FBTBuilder::AddTaskValuesChangedWithLogs(InnerSequence, Task2Execute, EBTNodeResult::Succeeded, /*IntValue1*/1, /*IntValue2*/INDEX_NONE);
+				FBTBuilder::WithDecorator<UBTDecorator_ForceSuccess>(InnerSequence);
+				
+				FBTBuilder::AddTask(InnerSequence, Task3Execute, EBTNodeResult::Succeeded);
+			}
+
+			FBTBuilder::AddTask(MainSequence, Task4Execute, EBTNodeResult::Succeeded);
+		}
+
+		ExpectedResult.Add(Task1Execute/*0*/);
+		ExpectedResult.Add(Task2Execute/*1*/);
+	}
+};
+IMPLEMENT_AI_LATENT_TEST(FAITest_BTDecoratorForceSuccessTaskFailedParentComposite, "System.AI.Behavior Trees.Decorator: force success task failed parent composite")
+
+struct FAITest_BTDecoratorForceSuccessCompositeFailed : public FAITest_SimpleBT
+{
+	FAITest_BTDecoratorForceSuccessCompositeFailed()
+	{
+		enum
+		{
+			Task1Execute,
+			Task2Execute,
+			Task3Execute,
+			Task4Execute,
+		};
+
+		UBTCompositeNode& MainSequence = FBTBuilder::AddSequence(*BTAsset);
+		{
+			FBTBuilder::AddTask(MainSequence, Task1Execute, EBTNodeResult::Succeeded);
+
+			UBTCompositeNode& InnerSequence = FBTBuilder::AddSequence(MainSequence);
+			{
+				FBTBuilder::WithDecorator<UBTDecorator_ForceSuccess>(MainSequence);
+				FBTBuilder::WithDecoratorBlackboard(MainSequence, EArithmeticKeyOperation::NotEqual, /*IntValue*/1, EBTFlowAbortMode::Self, EBTBlackboardRestart::ResultChange);
+
+				FBTBuilder::AddTaskValuesChangedWithLogs(InnerSequence, Task2Execute, EBTNodeResult::Succeeded, /*IntValue1*/1, /*IntValue2*/INDEX_NONE);
+
+				FBTBuilder::AddTask(InnerSequence, Task3Execute, EBTNodeResult::Succeeded);
+			}
+
+			FBTBuilder::AddTask(MainSequence, Task4Execute, EBTNodeResult::Succeeded);
+		}
+
+		ExpectedResult.Add(Task1Execute/*0*/);
+		ExpectedResult.Add(Task2Execute/*1*/);
+		ExpectedResult.Add(Task4Execute/*3*/);
+	}
+};
+IMPLEMENT_AI_LATENT_TEST(FAITest_BTDecoratorForceSuccessCompositeFailed, "System.AI.Behavior Trees.Decorator: force success composite failed self")
+
+struct FAITest_BTDecoratorForceSuccessTaskFailedSibling : public FAITest_SimpleBT
+{
+	FAITest_BTDecoratorForceSuccessTaskFailedSibling()
+	{
+		enum
+		{
+			Task1Execute,
+			Task2Execute,
+			Task3Execute,
+			Task4Execute,
+		};
+
+		UBTCompositeNode& MainSequence = FBTBuilder::AddSequence(*BTAsset);
+		{
+			FBTBuilder::AddTask(MainSequence, Task1Execute, EBTNodeResult::Succeeded);
+
+			UBTCompositeNode& InnerSequence = FBTBuilder::AddSequence(MainSequence);
+			{
+				FBTBuilder::AddTask(InnerSequence, Task2Execute, EBTNodeResult::Succeeded);
+				FBTBuilder::WithDecorator<UBTDecorator_ForceSuccess>(InnerSequence);
+
+				FBTBuilder::AddTask(InnerSequence, Task3Execute, EBTNodeResult::Failed);
+			}
+
+			FBTBuilder::AddTask(MainSequence, Task4Execute, EBTNodeResult::Succeeded);
+		}
+
+		ExpectedResult.Add(Task1Execute/*0*/);
+		ExpectedResult.Add(Task2Execute/*1*/);
+		ExpectedResult.Add(Task3Execute/*2*/);
+	}
+};
+IMPLEMENT_AI_LATENT_TEST(FAITest_BTDecoratorForceSuccessTaskFailedSibling, "System.AI.Behavior Trees.Decorator: force success task failed sibling task")
+
 struct FAITest_BTBasicParallelWait : public FAITest_SimpleBT
 {
 	FAITest_BTBasicParallelWait()
diff --git a/Engine/Source/Runtime/AIModule/Private/BehaviorTree/BehaviorTreeComponent.cpp b/Engine/Source/Runtime/AIModule/Private/BehaviorTree/BehaviorTreeComponent.cpp
--- a/Engine/Source/Runtime/AIModule/Private/BehaviorTree/BehaviorTreeComponent.cpp
+++ b/Engine/Source/Runtime/AIModule/Private/BehaviorTree/BehaviorTreeComponent.cpp
@@ -1950,6 +1950,9 @@ void UBehaviorTreeComponent::ProcessExecutionRequest()
 	SearchData.RollbackDeactivatedBranchStart = SearchData.DeactivatedBranchStart;
 	SearchData.RollbackDeactivatedBranchEnd = SearchData.DeactivatedBranchEnd;
 
+	// Setting search root node here it can be refer to when calling 'DeactivateUpTo'
+	SearchData.SearchRootNode = FBTNodeIndex(ExecutionRequest.ExecuteInstanceIdx, ExecutionRequest.ExecuteNode->GetExecutionIndex());
+
 	EBTNodeResult::Type NodeResult = ExecutionRequest.ContinueWithResult;
 	UBTTaskNode* NextTask = NULL;
 
@@ -1974,7 +1977,7 @@ void UBehaviorTreeComponent::ProcessExecutionRequest()
 			if (!bDeactivated)
 			{
 				// error occurred and tree will restart, all pending deactivation notifies will be lost
-				// this is should happen
+				// this should never happen
 
 				BT_SEARCHLOG(SearchData, Error, TEXT("Unable to deactivate up to %s. Active node is %s. All pending updates will be lost!"), 
 					*UBehaviorTreeTypes::DescribeNodeHelper(ExecutionRequest.ExecuteNode), 
@@ -2001,7 +2004,6 @@ void UBehaviorTreeComponent::ProcessExecutionRequest()
 		SearchData.AssignSearchId();
 		SearchData.bPostponeSearch = false;
 		SearchData.bSearchInProgress = true;
-		SearchData.SearchRootNode = FBTNodeIndex(ExecutionRequest.ExecuteInstanceIdx, ExecutionRequest.ExecuteNode->GetExecutionIndex());
 
 		// activate root node if needed (can't be handled by parent composite...)
 		if (ActiveInstance.ActiveNode == NULL)
diff --git a/Engine/Source/Runtime/AIModule/Private/BehaviorTree/Decorators/BTDecorator_ForceSuccess.cpp b/Engine/Source/Runtime/AIModule/Private/BehaviorTree/Decorators/BTDecorator_ForceSuccess.cpp
--- a/Engine/Source/Runtime/AIModule/Private/BehaviorTree/Decorators/BTDecorator_ForceSuccess.cpp
+++ b/Engine/Source/Runtime/AIModule/Private/BehaviorTree/Decorators/BTDecorator_ForceSuccess.cpp
@@ -1,6 +1,8 @@
 // Copyright Epic Games, Inc. All Rights Reserved.
 
 #include "BehaviorTree/Decorators/BTDecorator_ForceSuccess.h"
+
+#include "BehaviorTree/BTCompositeNode.h"
 #include "GameFramework/Actor.h"
 #include "VisualLogger/VisualLogger.h"
 
@@ -18,8 +20,17 @@ UBTDecorator_ForceSuccess::UBTDecorator_ForceSuccess(const FObjectInitializer& O
 
 void UBTDecorator_ForceSuccess::OnNodeProcessed(FBehaviorTreeSearchData& SearchData, EBTNodeResult::Type& NodeResult)
 {
-	NodeResult = EBTNodeResult::Succeeded;
-	BT_SEARCHLOG(SearchData, Log, TEXT("Forcing Success: %s"), *UBehaviorTreeTypes::DescribeNodeHelper(this));
+	const FBTNodeIndex ParentNodeIndex(SearchData.RollbackInstanceIdx, GetParentNode() != nullptr ? GetParentNode()->GetExecutionIndex() : 0);
+	const UBTNode* MyNode = GetMyNode();
+	const FBTNodeIndex MyNodeIndex(SearchData.RollbackInstanceIdx, MyNode != nullptr ? MyNode->GetExecutionIndex() : 0);
+	
+	// Only modify result if search originates from our parent node (abort self) or on our associated node (failure)
+	if (SearchData.SearchRootNode == ParentNodeIndex || SearchData.SearchRootNode == MyNodeIndex)
+	{
+		checkf(NodeResult != EBTNodeResult::Aborted, TEXT("Should never change a result set to 'Aborted'"));
+		NodeResult = EBTNodeResult::Succeeded;
+		BT_SEARCHLOG(SearchData, Log, TEXT("Forcing Success: %s"), *UBehaviorTreeTypes::DescribeNodeHelper(this));
+	}
 }
 
 #if WITH_EDITOR
