[BehaviorTree] ForceSuccess decorator effect should not be limited during the search

diff --git a/Engine/Source/Developer/AITestSuite/Private/Tests/BTTest.cpp b/Engine/Source/Developer/AITestSuite/Private/Tests/BTTest.cpp
--- a/Engine/Source/Developer/AITestSuite/Private/Tests/BTTest.cpp
+++ b/Engine/Source/Developer/AITestSuite/Private/Tests/BTTest.cpp
@@ -128,33 +128,112 @@ struct FAITest_BTBasicSelector : public FAITest_SimpleBT
 };
 IMPLEMENT_AI_LATENT_TEST(FAITest_BTBasicSelector, "System.AI.Behavior Trees.Composite node: selector")
 
-struct FAITest_BTBasicSequence : public FAITest_SimpleBT
+struct FAITest_BTSearchSequenceForceSuccessTask : public FAITest_SimpleBT
 {
-	FAITest_BTBasicSequence()
+	FAITest_BTSearchSequenceForceSuccessTask()
 	{
+		enum
+		{
+			Task1Execute = 1,
+			Task2Execute,
+			Task3Execute,
+			Task4Execute,
+		};
+
 		UBTCompositeNode& CompNode = FBTBuilder::AddSequence(*BTAsset);
 		{
-			FBTBuilder::AddTask(CompNode, 0, EBTNodeResult::Succeeded);
+			FBTBuilder::AddTask(CompNode, Task1Execute, EBTNodeResult::Succeeded);
 
-			FBTBuilder::AddTask(CompNode, 1, EBTNodeResult::Failed);
+			// First decorator will deny execution of Task 2 but ForceSuccess will only sequence to continue
+			FBTBuilder::AddTask(CompNode, Task2Execute, EBTNodeResult::Succeeded);
 			{
 				FBTBuilder::WithDecorator<UTestBTDecorator_CantExecute>(CompNode);
-
-				// @todo BT: We should reconsider allowing decorator to change the activation result.
-				// This looks more like a debugger feature than a desired behavior. 
 				FBTBuilder::WithDecorator<UBTDecorator_ForceSuccess>(CompNode);
 			}
 
-			FBTBuilder::AddTask(CompNode, 2, EBTNodeResult::Failed, 2);
+			FBTBuilder::AddTask(CompNode, Task3Execute, EBTNodeResult::Failed);
 
-			FBTBuilder::AddTask(CompNode, 3, EBTNodeResult::Succeeded);
+			FBTBuilder::AddTask(CompNode, Task4Execute, EBTNodeResult::Succeeded);
 		}
 
-		ExpectedResult.Add(0);
-		ExpectedResult.Add(2);
+		ExpectedResult.Add(Task1Execute);
+		ExpectedResult.Add(Task3Execute);
+	}
+};
+IMPLEMENT_AI_LATENT_TEST(FAITest_BTSearchSequenceForceSuccessTask, "System.AI.Behavior Trees.Composite node: search sequence with force success child task")
+
+struct FAITest_BTDeepSearchSequenceForceSuccessSelf : public FAITest_SimpleBT
+{
+	FAITest_BTDeepSearchSequenceForceSuccessSelf()
+	{
+		enum
+		{
+			Task1Execute = 1,
+			Task2Execute,
+		};
+
+		UBTCompositeNode& Selector = FBTBuilder::AddSelector(*BTAsset);
+		{
+			UBTCompositeNode& MainSequence = FBTBuilder::AddSequence(Selector);
+			{
+				UBTCompositeNode& InnerSequence = FBTBuilder::AddSequence(MainSequence);
+				{
+					// Force success must allow sequence to continue even if the search root node is not its associate node or its immediate parent.
+					FBTBuilder::WithDecorator<UBTDecorator_ForceSuccess>(MainSequence);
+					FBTBuilder::WithDecorator<UTestBTDecorator_CantExecute>(MainSequence);
+
+					FBTBuilder::AddTask(InnerSequence, Task1Execute, EBTNodeResult::Succeeded);
+				}
+
+				FBTBuilder::AddTask(MainSequence, Task2Execute, EBTNodeResult::Succeeded);
+			}
+		}
+
+		ExpectedResult.Add(Task2Execute);
+	}
+};
+IMPLEMENT_AI_LATENT_TEST(FAITest_BTDeepSearchSequenceForceSuccessSelf, "System.AI.Behavior Trees.Composite node: deep search sequence with force success self")
+
+struct FAITest_BTSearchParentNotAffectedByForceSuccessChild : public FAITest_SimpleBT
+{
+	FAITest_BTSearchParentNotAffectedByForceSuccessChild()
+	{
+		enum
+		{
+			Task1Execute = 1,
+			Task2Execute,
+			Task3Execute,
+			Task4Execute,
+		};
+
+		UBTCompositeNode& MainSequence = FBTBuilder::AddSelector(*BTAsset);
+		{
+			UBTCompositeNode& MiddleSequence = FBTBuilder::AddSequence(MainSequence);
+			{
+				UBTCompositeNode& InnerSequence = FBTBuilder::AddSequence(MiddleSequence);
+				{
+					FBTBuilder::AddTask(InnerSequence, Task1Execute, EBTNodeResult::Succeeded);
+					{
+						FBTBuilder::WithDecorator<UBTDecorator_ForceSuccess>(InnerSequence);
+						FBTBuilder::WithDecorator<UTestBTDecorator_CantExecute>(InnerSequence);
+					}
+
+					FBTBuilder::AddTask(MiddleSequence, Task2Execute, EBTNodeResult::Succeeded);
+					{
+						FBTBuilder::WithDecorator<UTestBTDecorator_CantExecute>(MiddleSequence);
+					}
+				}
+				
+				FBTBuilder::AddTask(MiddleSequence, Task3Execute, EBTNodeResult::Succeeded);
+			}
+
+			FBTBuilder::AddTask(MainSequence, Task4Execute, EBTNodeResult::Succeeded);
+		}
+
+		ExpectedResult.Add(Task4Execute);
 	}
 };
-IMPLEMENT_AI_LATENT_TEST(FAITest_BTBasicSequence, "System.AI.Behavior Trees.Composite node: sequence")
+IMPLEMENT_AI_LATENT_TEST(FAITest_BTSearchParentNotAffectedByForceSuccessChild, "System.AI.Behavior Trees.Composite node: search in parent node not affected by force success child")
 
 struct FAITest_BTDecoratorBlueprint : public FAITest_SimpleBT
 {
@@ -240,7 +319,7 @@ struct FAITest_BTDecoratorForceSuccessTaskFailed : public FAITest_SimpleBT
 	{
 		enum
 		{
-			Task1Execute,
+			Task1Execute = 1,
 			Task2Execute,
 			Task3Execute,
 			Task4Execute,
@@ -261,10 +340,10 @@ struct FAITest_BTDecoratorForceSuccessTaskFailed : public FAITest_SimpleBT
 			FBTBuilder::AddTask(MainSequence, Task4Execute, EBTNodeResult::Succeeded);
 		}
 
-		ExpectedResult.Add(Task1Execute/*0*/);
-		ExpectedResult.Add(Task2Execute/*1*/);
-		ExpectedResult.Add(Task3Execute/*2*/);
-		ExpectedResult.Add(Task4Execute/*3*/);
+		ExpectedResult.Add(Task1Execute);
+		ExpectedResult.Add(Task2Execute);
+		ExpectedResult.Add(Task3Execute);
+		ExpectedResult.Add(Task4Execute);
 	}
 };
 IMPLEMENT_AI_LATENT_TEST(FAITest_BTDecoratorForceSuccessTaskFailed, "System.AI.Behavior Trees.Decorator: force success task failed self")
@@ -275,7 +354,7 @@ struct FAITest_BTDecoratorForceSuccessTaskFailedParentComposite : public FAITest
 	{
 		enum
 		{
-			Task1Execute,
+			Task1Execute = 1,
 			Task2Execute,
 			Task3Execute,
 			Task4Execute,
@@ -290,7 +369,9 @@ struct FAITest_BTDecoratorForceSuccessTaskFailedParentComposite : public FAITest
 				FBTBuilder::WithDecoratorBlackboard(MainSequence, EArithmeticKeyOperation::NotEqual, /*IntValue*/1, EBTFlowAbortMode::Self, EBTBlackboardRestart::ResultChange);
 
 				FBTBuilder::AddTaskValuesChangedWithLogs(InnerSequence, Task2Execute, EBTNodeResult::Succeeded, /*IntValue1*/1, /*IntValue2*/INDEX_NONE);
-				FBTBuilder::WithDecorator<UBTDecorator_ForceSuccess>(InnerSequence);
+				{
+					FBTBuilder::WithDecorator<UBTDecorator_ForceSuccess>(InnerSequence);
+				}
 				
 				FBTBuilder::AddTask(InnerSequence, Task3Execute, EBTNodeResult::Succeeded);
 			}
@@ -298,8 +379,8 @@ struct FAITest_BTDecoratorForceSuccessTaskFailedParentComposite : public FAITest
 			FBTBuilder::AddTask(MainSequence, Task4Execute, EBTNodeResult::Succeeded);
 		}
 
-		ExpectedResult.Add(Task1Execute/*0*/);
-		ExpectedResult.Add(Task2Execute/*1*/);
+		ExpectedResult.Add(Task1Execute);
+		ExpectedResult.Add(Task2Execute);
 	}
 };
 IMPLEMENT_AI_LATENT_TEST(FAITest_BTDecoratorForceSuccessTaskFailedParentComposite, "System.AI.Behavior Trees.Decorator: force success task failed parent composite")
@@ -310,7 +391,7 @@ struct FAITest_BTDecoratorForceSuccessCompositeFailed : public FAITest_SimpleBT
 	{
 		enum
 		{
-			Task1Execute,
+			Task1Execute=1,
 			Task2Execute,
 			Task3Execute,
 			Task4Execute,
@@ -333,20 +414,92 @@ struct FAITest_BTDecoratorForceSuccessCompositeFailed : public FAITest_SimpleBT
 			FBTBuilder::AddTask(MainSequence, Task4Execute, EBTNodeResult::Succeeded);
 		}
 
-		ExpectedResult.Add(Task1Execute/*0*/);
-		ExpectedResult.Add(Task2Execute/*1*/);
-		ExpectedResult.Add(Task4Execute/*3*/);
+		ExpectedResult.Add(Task1Execute);
+		ExpectedResult.Add(Task2Execute);
+		ExpectedResult.Add(Task4Execute);
 	}
 };
 IMPLEMENT_AI_LATENT_TEST(FAITest_BTDecoratorForceSuccessCompositeFailed, "System.AI.Behavior Trees.Decorator: force success composite failed self")
 
+struct FAITest_BTDecoratorForceSuccessCompositeFailedFromLastChild : public FAITest_SimpleBT
+{
+	FAITest_BTDecoratorForceSuccessCompositeFailedFromLastChild()
+	{
+		enum
+		{
+			Task1Execute=1,
+			Task2Execute,
+			Task3Execute,
+			Task4Execute,
+		};
+
+		UBTCompositeNode& MainSequence = FBTBuilder::AddSequence(*BTAsset);
+		{
+			FBTBuilder::AddTask(MainSequence, Task1Execute, EBTNodeResult::Succeeded);
+
+			UBTCompositeNode& InnerSequence = FBTBuilder::AddSequence(MainSequence);
+			{
+				FBTBuilder::WithDecorator<UBTDecorator_ForceSuccess>(MainSequence);
+				FBTBuilder::WithDecoratorBlackboard(MainSequence, EArithmeticKeyOperation::NotEqual, /*IntValue*/1, EBTFlowAbortMode::Both, EBTBlackboardRestart::ResultChange);
+
+				FBTBuilder::AddTask(InnerSequence, Task2Execute, EBTNodeResult::Succeeded);
+
+				FBTBuilder::AddTaskValuesChangedWithLogs(InnerSequence, Task3Execute, EBTNodeResult::Succeeded, /*IntValue1*/1, /*IntValue2*/INDEX_NONE);
+			}
+
+			FBTBuilder::AddTask(MainSequence, Task4Execute, EBTNodeResult::Succeeded);
+		}
+
+		ExpectedResult.Add(Task1Execute);
+		ExpectedResult.Add(Task2Execute);
+		ExpectedResult.Add(Task3Execute);
+		ExpectedResult.Add(Task4Execute);
+	}
+};
+IMPLEMENT_AI_LATENT_TEST(FAITest_BTDecoratorForceSuccessCompositeFailedFromLastChild, "System.AI.Behavior Trees.Decorator: force success composite failed self from last child")
+
+struct FAITest_BTDecoratorForceSuccessCompositeFailedChild : public FAITest_SimpleBT
+{
+	FAITest_BTDecoratorForceSuccessCompositeFailedChild()
+	{
+		enum
+		{
+			Task1Execute = 1,
+			Task2Execute,
+			Task3Execute,
+			Task4Execute,
+		};
+
+		UBTCompositeNode& MainSequence = FBTBuilder::AddSequence(*BTAsset);
+		{
+			FBTBuilder::AddTask(MainSequence, Task1Execute, EBTNodeResult::Succeeded);
+
+			UBTCompositeNode& InnerSequence = FBTBuilder::AddSequence(MainSequence);
+			{
+				FBTBuilder::WithDecorator<UBTDecorator_ForceSuccess>(MainSequence);
+
+				FBTBuilder::AddTask(InnerSequence, Task2Execute, EBTNodeResult::Failed);
+
+				FBTBuilder::AddTask(InnerSequence, Task3Execute, EBTNodeResult::Succeeded);
+			}
+
+			FBTBuilder::AddTask(MainSequence, Task4Execute, EBTNodeResult::Succeeded);
+		}
+
+		ExpectedResult.Add(Task1Execute);
+		ExpectedResult.Add(Task2Execute);
+		ExpectedResult.Add(Task4Execute);
+	}
+};
+IMPLEMENT_AI_LATENT_TEST(FAITest_BTDecoratorForceSuccessCompositeFailedChild, "System.AI.Behavior Trees.Decorator: force success composite failed child task")
+
 struct FAITest_BTDecoratorForceSuccessTaskFailedSibling : public FAITest_SimpleBT
 {
 	FAITest_BTDecoratorForceSuccessTaskFailedSibling()
 	{
 		enum
 		{
-			Task1Execute,
+			Task1Execute = 1,
 			Task2Execute,
 			Task3Execute,
 			Task4Execute,
@@ -359,7 +512,9 @@ struct FAITest_BTDecoratorForceSuccessTaskFailedSibling : public FAITest_SimpleB
 			UBTCompositeNode& InnerSequence = FBTBuilder::AddSequence(MainSequence);
 			{
 				FBTBuilder::AddTask(InnerSequence, Task2Execute, EBTNodeResult::Succeeded);
-				FBTBuilder::WithDecorator<UBTDecorator_ForceSuccess>(InnerSequence);
+				{
+					FBTBuilder::WithDecorator<UBTDecorator_ForceSuccess>(InnerSequence);
+				}
 
 				FBTBuilder::AddTask(InnerSequence, Task3Execute, EBTNodeResult::Failed);
 			}
@@ -367,9 +522,9 @@ struct FAITest_BTDecoratorForceSuccessTaskFailedSibling : public FAITest_SimpleB
 			FBTBuilder::AddTask(MainSequence, Task4Execute, EBTNodeResult::Succeeded);
 		}
 
-		ExpectedResult.Add(Task1Execute/*0*/);
-		ExpectedResult.Add(Task2Execute/*1*/);
-		ExpectedResult.Add(Task3Execute/*2*/);
+		ExpectedResult.Add(Task1Execute);
+		ExpectedResult.Add(Task2Execute);
+		ExpectedResult.Add(Task3Execute);
 	}
 };
 IMPLEMENT_AI_LATENT_TEST(FAITest_BTDecoratorForceSuccessTaskFailedSibling, "System.AI.Behavior Trees.Decorator: force success task failed sibling task")
diff --git a/Engine/Source/Runtime/AIModule/Private/BehaviorTree/Decorators/BTDecorator_ForceSuccess.cpp b/Engine/Source/Runtime/AIModule/Private/BehaviorTree/Decorators/BTDecorator_ForceSuccess.cpp
--- a/Engine/Source/Runtime/AIModule/Private/BehaviorTree/Decorators/BTDecorator_ForceSuccess.cpp
+++ b/Engine/Source/Runtime/AIModule/Private/BehaviorTree/Decorators/BTDecorator_ForceSuccess.cpp
@@ -20,12 +20,20 @@ UBTDecorator_ForceSuccess::UBTDecorator_ForceSuccess(const FObjectInitializer& O
 
 void UBTDecorator_ForceSuccess::OnNodeProcessed(FBehaviorTreeSearchData& SearchData, EBTNodeResult::Type& NodeResult)
 {
-	const FBTNodeIndex ParentNodeIndex(SearchData.RollbackInstanceIdx, GetParentNode() != nullptr ? GetParentNode()->GetExecutionIndex() : 0);
-	const UBTNode* MyNode = GetMyNode();
-	const FBTNodeIndex MyNodeIndex(SearchData.RollbackInstanceIdx, MyNode != nullptr ? MyNode->GetExecutionIndex() : 0);
-	
-	// Only modify result if search originates from our parent node (abort self) or on our associated node (failure)
-	if (SearchData.SearchRootNode == ParentNodeIndex || SearchData.SearchRootNode == MyNodeIndex)
+	bool bCanForceSuccess = true;
+
+	// Decorator is always allowed to force success during the search to ignore an optional branch in a sequence.
+	// But when used to override the node result on failure we only modify result if search originates
+	// from our parent node (abort self) or on our associated node (failure).
+	if (!SearchData.bSearchInProgress)
+	{
+		const FBTNodeIndex ParentNodeIndex(SearchData.RollbackInstanceIdx, GetParentNode() != nullptr ? GetParentNode()->GetExecutionIndex() : 0);
+		const UBTNode* MyNode = GetMyNode();
+		const FBTNodeIndex MyNodeIndex(SearchData.RollbackInstanceIdx, MyNode != nullptr ? MyNode->GetExecutionIndex() : 0);
+		bCanForceSuccess = (SearchData.SearchRootNode == ParentNodeIndex || SearchData.SearchRootNode == MyNodeIndex);
+	}
+
+	if (bCanForceSuccess)
 	{
 		checkf(NodeResult != EBTNodeResult::Aborted, TEXT("Should never change a result set to 'Aborted'"));
 		NodeResult = EBTNodeResult::Succeeded;
