AITest: Add CA_ASSUME to the AITEST_NOT_NULL so that static analyzer can treat it as null check

diff --git a/Engine/Source/Developer/AITestSuite/Public/AITestsCommon.h b/Engine/Source/Developer/AITestSuite/Public/AITestsCommon.h
--- a/Engine/Source/Developer/AITestSuite/Public/AITestsCommon.h
+++ b/Engine/Source/Developer/AITestSuite/Public/AITestsCommon.h
@@ -209,7 +209,9 @@ struct FAITest_SimpleComponentBasedTest : public FAITestBase
 #define AITEST_TRUE(What, Value) __AITEST_IMPL(What, Value, TestTrue, false)
 #define AITEST_FALSE(What, Value) __AITEST_IMPL(What, Value, TestFalse, false)
 #define AITEST_NULL(What, Pointer) __AITEST_IMPL(What, Pointer, TestNull, false)
-#define AITEST_NOT_NULL(What, Pointer) __AITEST_IMPL(What, Pointer, TestNotNull, false)
+#define AITEST_NOT_NULL(What, Pointer) \
+	__AITEST_IMPL(What, Pointer, TestNotNull, false) \
+	CA_ASSUME(Pointer)
 
 namespace FTestHelpers
 {
