[BehaviorTree] Updated TimeLimit Decorator to not get get ticked again in case it is still active after requesting execution (e.g. latent abort)
Schedule an execution when a task is completed and there is a pending execution.

diff --git a/Engine/Source/Developer/AITestSuite/Classes/BehaviorTree/TestBTTask_TimerBasedLatent.h b/Engine/Source/Developer/AITestSuite/Classes/BehaviorTree/TestBTTask_TimerBasedLatent.h
new file mode 100644
--- /dev/null
+++ b/Engine/Source/Developer/AITestSuite/Classes/BehaviorTree/TestBTTask_TimerBasedLatent.h
@@ -0,0 +1,52 @@
+// Copyright Epic Games, Inc. All Rights Reserved.
+#pragma once
+
+#include "BehaviorTree/BTTaskNode.h"
+#include "TestBTTask_TimerBasedLatent.generated.h"
+
+struct FBTTimerBasedLatentTaskMemory
+{
+	FTimerHandle TimerHandle;
+	uint8 bIsAborting : 1;
+};
+
+UCLASS(meta = (HiddenNode))
+class UTestBTTask_TimerBasedLatent : public UBTTaskNode
+{
+	GENERATED_BODY()
+
+public:
+
+	UTestBTTask_TimerBasedLatent();
+
+	UPROPERTY()
+	int32 LogIndexExecuteStart = INDEX_NONE;
+
+	UPROPERTY()
+	int32 LogIndexExecuteFinish = INDEX_NONE;
+
+	UPROPERTY()
+	int32 LogIndexAbortStart = INDEX_NONE;
+
+	UPROPERTY()
+	int32 LogIndexAbortFinish = INDEX_NONE;
+
+	/** Num of ticks from 'execute start' to 'execute finish' */
+	UPROPERTY()
+	int32 NumTicksExecuting = 0;
+
+	/** Num of ticks from 'abort start' to 'abort finish' */
+	UPROPERTY()
+	int32 NumTicksAborting = 0;
+
+	UPROPERTY()
+	TEnumAsByte<EBTNodeResult::Type> LogResult = EBTNodeResult::Succeeded;
+
+protected:
+	virtual EBTNodeResult::Type ExecuteTask(UBehaviorTreeComponent& OwnerComp, uint8* NodeMemory) override;
+	virtual EBTNodeResult::Type AbortTask(UBehaviorTreeComponent& OwnerComp, uint8* NodeMemory) override;
+	virtual uint16 GetInstanceMemorySize() const override;
+	virtual void InitializeMemory(UBehaviorTreeComponent& OwnerComp, uint8* NodeMemory, EBTMemoryInit::Type InitType) const override;
+	virtual void CleanupMemory(UBehaviorTreeComponent& OwnerComp, uint8* NodeMemory, EBTMemoryClear::Type CleanupType) const override;
+	void LogExecution(UBehaviorTreeComponent& OwnerComp, int32 LogNumber) const;
+};
\ No newline at end of file
diff --git a/Engine/Source/Developer/AITestSuite/Private/BehaviorTree/TestBTTask_TimerBasedLatent.cpp b/Engine/Source/Developer/AITestSuite/Private/BehaviorTree/TestBTTask_TimerBasedLatent.cpp
new file mode 100644
--- /dev/null
+++ b/Engine/Source/Developer/AITestSuite/Private/BehaviorTree/TestBTTask_TimerBasedLatent.cpp
@@ -0,0 +1,103 @@
+// Copyright Epic Games, Inc. All Rights Reserved.
+
+#include "BehaviorTree/TestBTTask_TimerBasedLatent.h"
+#include "Engine/World.h"
+#include "MockAI_BT.h"
+#include "TimerManager.h"
+
+#include UE_INLINE_GENERATED_CPP_BY_NAME(TestBTTask_TimerBasedLatent)
+
+UTestBTTask_TimerBasedLatent::UTestBTTask_TimerBasedLatent()
+{
+	INIT_TASK_NODE_NOTIFY_FLAGS();
+}
+
+EBTNodeResult::Type UTestBTTask_TimerBasedLatent::ExecuteTask(UBehaviorTreeComponent& OwnerComp, uint8* NodeMemory)
+{
+	FBTTimerBasedLatentTaskMemory* TaskMemory = CastInstanceNodeMemory<FBTTimerBasedLatentTaskMemory>(NodeMemory);
+	TaskMemory->bIsAborting = false;
+
+	LogExecution(OwnerComp, LogIndexExecuteStart);
+	if (NumTicksExecuting == 0)
+	{
+		LogExecution(OwnerComp, LogIndexExecuteFinish);
+		return LogResult;
+	}
+
+	if (const UWorld* World = OwnerComp.GetWorld())
+	{
+		const float DeltaTime = NumTicksExecuting * FAITestHelpers::TickInterval;
+		World->GetTimerManager().SetTimer(TaskMemory->TimerHandle,
+			FTimerDelegate::CreateWeakLambda(this, [&OwnerComp, TaskMemory, this]()
+			{
+				TaskMemory->TimerHandle.Invalidate();
+
+				ensure(!TaskMemory->bIsAborting);
+				LogExecution(OwnerComp, LogIndexExecuteFinish);
+				FinishLatentTask(OwnerComp, LogResult);
+			}), DeltaTime, false);
+	}
+
+	return EBTNodeResult::InProgress;
+}
+
+EBTNodeResult::Type UTestBTTask_TimerBasedLatent::AbortTask(UBehaviorTreeComponent& OwnerComp, uint8* NodeMemory)
+{
+	FBTTimerBasedLatentTaskMemory* TaskMemory = CastInstanceNodeMemory<FBTTimerBasedLatentTaskMemory>(NodeMemory);
+	TaskMemory->bIsAborting = true;
+
+	LogExecution(OwnerComp, LogIndexAbortStart);
+	if (NumTicksAborting == 0)
+	{
+		LogExecution(OwnerComp, LogIndexAbortFinish);
+		return EBTNodeResult::Aborted;
+	}
+
+	if (const UWorld* World = OwnerComp.GetWorld())
+	{
+		const float DeltaTime = NumTicksAborting * FAITestHelpers::TickInterval;
+		World->GetTimerManager().SetTimer(TaskMemory->TimerHandle,
+			FTimerDelegate::CreateWeakLambda(this, [&OwnerComp, TaskMemory, this]()
+			{
+				TaskMemory->TimerHandle.Invalidate();
+
+				ensure(TaskMemory->bIsAborting);
+				LogExecution(OwnerComp, LogIndexAbortFinish);
+				FinishLatentAbort(OwnerComp);
+			}), DeltaTime, false);
+	}
+
+	return EBTNodeResult::InProgress;
+}
+
+uint16 UTestBTTask_TimerBasedLatent::GetInstanceMemorySize() const
+{
+	return sizeof(FBTTimerBasedLatentTaskMemory);
+}
+
+void UTestBTTask_TimerBasedLatent::InitializeMemory(UBehaviorTreeComponent& OwnerComp, uint8* NodeMemory, const EBTMemoryInit::Type InitType) const
+{
+	InitializeNodeMemory<FBTTimerBasedLatentTaskMemory>(NodeMemory, InitType);
+}
+
+void UTestBTTask_TimerBasedLatent::CleanupMemory(UBehaviorTreeComponent& OwnerComp, uint8* NodeMemory, const EBTMemoryClear::Type CleanupType) const
+{
+	const FBTTimerBasedLatentTaskMemory* TaskMemory = CastInstanceNodeMemory<FBTTimerBasedLatentTaskMemory>(NodeMemory);
+	if (TaskMemory->TimerHandle.IsValid())
+	{
+		if (const UWorld* World = OwnerComp.GetWorld())
+		{
+			World->GetTimerManager().ClearTimer(CastInstanceNodeMemory<FBTTimerBasedLatentTaskMemory>(NodeMemory)->TimerHandle);
+		}
+	}
+
+	CleanupNodeMemory<FBTTimerBasedLatentTaskMemory>(NodeMemory, CleanupType);
+}
+
+void UTestBTTask_TimerBasedLatent::LogExecution(UBehaviorTreeComponent&, const int32 LogNumber) const
+{
+	if (LogNumber >= 0)
+	{
+		UMockAI_BT::ExecutionLog.Add(LogNumber);
+	}
+}
diff --git a/Engine/Source/Developer/AITestSuite/Private/Tests/BTTest.cpp b/Engine/Source/Developer/AITestSuite/Private/Tests/BTTest.cpp
--- a/Engine/Source/Developer/AITestSuite/Private/Tests/BTTest.cpp
+++ b/Engine/Source/Developer/AITestSuite/Private/Tests/BTTest.cpp
@@ -7,6 +7,7 @@
 #include "AITestsCommon.h"
 #include "MockAI_BT.h"
 #include "BehaviorTree/TestBTDecorator_CantExecute.h"
+#include "BehaviorTree/TestBTTask_TimerBasedLatent.h"
 #include "BehaviorTree/Decorators/BTDecorator_TimeLimit.h"
 
 #define LOCTEXT_NAMESPACE "AITestSuite_BTTest"
@@ -1400,6 +1401,53 @@ struct FAITest_BTRequestExecutionOnLatentAbortFinished : public FAITest_SimpleBT
 };
 IMPLEMENT_AI_LATENT_TEST(FAITest_BTRequestExecutionOnLatentAbortFinished, "System.AI.Behavior Trees.Abort: request on latent task finished")
 
+struct FAITest_BTRequestExecutionOnTimerBasedLatentAbortFinished : public FAITest_SimpleBT
+{
+	FAITest_BTRequestExecutionOnTimerBasedLatentAbortFinished()
+	{
+		enum
+		{
+			FirstTaskStartExecute = 1,
+			FirstTaskStartAbort,
+			FirstTaskFinishAbort,
+			FirstTaskFinished,
+			SecondTaskExecute
+		};
+
+		UBTCompositeNode& CompNode = FBTBuilder::AddSelector(*BTAsset);
+		{
+			UBTCompositeNode& SelectorNode = FBTBuilder::AddSelector(CompNode);
+			{
+				constexpr float NumTicksBeforeAbortingTask = 3;
+
+				// Add this decorator to prevent the test from running indefinitely if there is a regression on the behavior
+				// this unit test is meant to validate 
+				constexpr float NumTicksBeforeAbortingSelector = NumTicksBeforeAbortingTask + 10;
+				FBTBuilder::WithDecorator<UBTDecorator_TimeLimit>(CompNode).TimeLimit = NumTicksBeforeAbortingSelector * FAITestHelpers::TickInterval;
+
+				UTestBTTask_TimerBasedLatent& Task = FBTBuilder::AddTask<UTestBTTask_TimerBasedLatent>(SelectorNode);
+				Task.LogIndexExecuteStart = FirstTaskStartExecute;
+				Task.LogIndexAbortStart = FirstTaskStartAbort;
+				Task.LogIndexAbortFinish = FirstTaskFinishAbort;
+				Task.NumTicksAborting = 2;
+				// Make sure to execute long enough to be aborted
+				Task.NumTicksExecuting = 2 * NumTicksBeforeAbortingTask;
+				{
+					FBTBuilder::WithDecorator<UBTDecorator_TimeLimit>(SelectorNode).TimeLimit = NumTicksBeforeAbortingTask * FAITestHelpers::TickInterval;
+				}
+
+				FBTBuilder::AddTask(SelectorNode, SecondTaskExecute, EBTNodeResult::Succeeded);
+			}
+		}
+
+		ExpectedResult.Add(FirstTaskStartExecute);
+		ExpectedResult.Add(FirstTaskStartAbort);
+		ExpectedResult.Add(FirstTaskFinishAbort);
+		ExpectedResult.Add(SecondTaskExecute);
+	}
+};
+IMPLEMENT_AI_LATENT_TEST(FAITest_BTRequestExecutionOnTimerBasedLatentAbortFinished, "System.AI.Behavior Trees.Abort: request on latent timer based task finished")
+
 struct FAITest_BTSwitchingHigherPrioDuringServiceBecomeRelevant : public FAITest_SimpleBT
 {
 	FAITest_BTSwitchingHigherPrioDuringServiceBecomeRelevant()
diff --git a/Engine/Source/Developer/AITestSuite/Public/BTBuilder.h b/Engine/Source/Developer/AITestSuite/Public/BTBuilder.h
--- a/Engine/Source/Developer/AITestSuite/Public/BTBuilder.h
+++ b/Engine/Source/Developer/AITestSuite/Public/BTBuilder.h
@@ -134,6 +134,15 @@ struct FBTBuilder
 		return *NodeOb;
 	}
 
+	template<class T>
+	static T& AddTask(UBTCompositeNode& ParentNode, UClass* TaskClass = T::StaticClass())
+	{
+		T* TaskNode = NewObject<T>(ParentNode.GetTreeAsset());
+		ParentNode.Children.Emplace_GetRef().ChildTask = TaskNode;
+
+		return *TaskNode;
+	}
+
 	static void AddTask(UBTCompositeNode& ParentNode, int32 LogIndex, EBTNodeResult::Type NodeResult, int32 ExecutionTicks = 0, int32 LogTickIndex = -1)
 	{
 		UTestBTTask_Log* TaskNode = NewObject<UTestBTTask_Log>(ParentNode.GetTreeAsset());
diff --git a/Engine/Source/Runtime/AIModule/Private/BehaviorTree/BehaviorTreeComponent.cpp b/Engine/Source/Runtime/AIModule/Private/BehaviorTree/BehaviorTreeComponent.cpp
--- a/Engine/Source/Runtime/AIModule/Private/BehaviorTree/BehaviorTreeComponent.cpp
+++ b/Engine/Source/Runtime/AIModule/Private/BehaviorTree/BehaviorTreeComponent.cpp
@@ -546,6 +546,11 @@ void UBehaviorTreeComponent::OnTaskFinished(const UBTTaskNode* TaskNode, EBTNode
 			{
 				RequestExecution(TaskResult);
 			}
+			// schedule an update to process a pending execution when task is done aborting
+			else if (PendingExecution.IsSet())
+			{
+				ScheduleExecutionUpdate();
+			}
 		}
 		else if (TaskResult == EBTNodeResult::Aborted && InstanceStack.IsValidIndex(TaskInstanceIdx) && InstanceStack[TaskInstanceIdx].ActiveNode == TaskNode)
 		{
@@ -1696,7 +1701,7 @@ void UBehaviorTreeComponent::TickComponent(float DeltaTime, const ELevelTick Tic
 	if (NextTickDeltaTime > 0.0f)
 	{
 		// The TickManager is using global time to calculate delta since last ticked time. When the value is big, we can get into float precision errors compare to our calculation.
-		if (NextTickDeltaTime > KINDA_SMALL_NUMBER)
+		if (ThisTickFunction != nullptr && NextTickDeltaTime > KINDA_SMALL_NUMBER)
 		{
 			UE_VLOG(GetOwner(), LogBehaviorTree, Error, TEXT("BT(%i) did not need to be tick, ask deltatime of %fs got %fs with a diff of %fs."), GFrameCounter, NextTickDeltaTime + AccumulatedTickDeltaTime + DeltaTime, DeltaTime + AccumulatedTickDeltaTime, NextTickDeltaTime);
 		}
diff --git a/Engine/Source/Runtime/AIModule/Private/BehaviorTree/Decorators/BTDecorator_TimeLimit.cpp b/Engine/Source/Runtime/AIModule/Private/BehaviorTree/Decorators/BTDecorator_TimeLimit.cpp
--- a/Engine/Source/Runtime/AIModule/Private/BehaviorTree/Decorators/BTDecorator_TimeLimit.cpp
+++ b/Engine/Source/Runtime/AIModule/Private/BehaviorTree/Decorators/BTDecorator_TimeLimit.cpp
@@ -63,8 +63,13 @@ void UBTDecorator_TimeLimit::TickNode(UBehaviorTreeComponent& OwnerComp, uint8*
 		DeltaSeconds,
 		TimeLimit);
 
+	// Mark this decorator instance as Elapsed for calls to CalculateRawConditionValue
 	reinterpret_cast<FBTimeLimitMemory*>(NodeMemory)->bElapsed = true;
 
+	// Set our next tick time to large value so we don't get ticked again in case the decorator
+	// is still active after requesting execution (e.g. latent abort)
+	SetNextTickTime(NodeMemory, FLT_MAX);
+	
 	OwnerComp.RequestExecution(this);
 }
 
