[BTTest] prevent BehaviorTree component used by MockAI_BT from being updated by the component's TickFunction so it is always updated by the test framework.
Also handling updates for the world's TimerManager for cases where the world is not updated by the engine tick (e.g., editor world not ticked if viewport is not set to 'realtime')

diff --git a/Engine/Source/Developer/AITestSuite/Private/MockAI/MockAI_BT.cpp b/Engine/Source/Developer/AITestSuite/Private/MockAI/MockAI_BT.cpp
--- a/Engine/Source/Developer/AITestSuite/Private/MockAI/MockAI_BT.cpp
+++ b/Engine/Source/Developer/AITestSuite/Private/MockAI/MockAI_BT.cpp
@@ -21,6 +21,9 @@ UMockAI_BT::UMockAI_BT(const FObjectInitializer& ObjectInitializer)
 		UseBrainComponent<UBehaviorTreeComponent>();
 
 		BTComp = Cast<UBehaviorTreeComponent>(BrainComp);
+
+		// We want our component to be ticked manually by the test framework, never by the TickFunction
+		BTComp->PrimaryComponentTick.bCanEverTick = false;
 		BTComp->RegisterComponent();
 		BTComp->InitializeComponent();
 	}
diff --git a/Engine/Source/Developer/AITestSuite/Private/Tests/BTTest.cpp b/Engine/Source/Developer/AITestSuite/Private/Tests/BTTest.cpp
--- a/Engine/Source/Developer/AITestSuite/Private/Tests/BTTest.cpp
+++ b/Engine/Source/Developer/AITestSuite/Private/Tests/BTTest.cpp
@@ -1,14 +1,15 @@
 // Copyright Epic Games, Inc. All Rights Reserved.
 
-#include "CoreMinimal.h"
+#include "AITestsCommon.h"
+#include "BTBuilder.h"
 #include "BehaviorTree/BTCompositeNode.h"
 #include "BehaviorTree/Decorators/BTDecorator_ForceSuccess.h"
-#include "BTBuilder.h"
-#include "AITestsCommon.h"
-#include "MockAI_BT.h"
+#include "BehaviorTree/Decorators/BTDecorator_TimeLimit.h"
 #include "BehaviorTree/TestBTDecorator_CantExecute.h"
 #include "BehaviorTree/TestBTTask_TimerBasedLatent.h"
-#include "BehaviorTree/Decorators/BTDecorator_TimeLimit.h"
+#include "Engine/World.h"
+#include "MockAI_BT.h"
+#include "TimerManager.h"
 
 #define LOCTEXT_NAMESPACE "AITestSuite_BTTest"
 
@@ -54,8 +55,16 @@ struct FAITest_SimpleBT : public FAITestBase
 	{
 		FAITestHelpers::UpdateFrameCounter();
 
-		if (AIBTUser != NULL)
+		if (AIBTUser != nullptr)
 		{
+			// Some tests rely on timers, so we need to manually tick the world timer manager
+			// if it was not ticked in the frame (e.g., the editor world not ticking if viewport is not set to 'realtime')
+			FTimerManager& TimerManager = GetWorld().GetTimerManager();
+			if (!TimerManager.HasBeenTickedThisFrame())
+			{
+				TimerManager.Tick(FAITestHelpers::TickInterval);
+			}
+
 			if (bUseSystemTicking == false)
 			{
 				AIBTUser->TickMe(FAITestHelpers::TickInterval);
