Initial implementation of AITestSuite-testable ensures and checks

diff --git a/Engine/Source/Developer/AITestSuite/AITestSuite.Build.cs b/Engine/Source/Developer/AITestSuite/AITestSuite.Build.cs
--- a/Engine/Source/Developer/AITestSuite/AITestSuite.Build.cs
+++ b/Engine/Source/Developer/AITestSuite/AITestSuite.Build.cs
@@ -12,12 +12,14 @@ namespace UnrealBuildTool.Rules
 					"Core",
 					"CoreUObject",
 					"Engine",
-                    "GameplayTasks",
-                    "AIModule",
-                }
-                );
+					"GameplayTasks",
+					"AIModule",
+				}
+				);
 
 			PrecompileForTargets = PrecompileTargetsType.Any;
+
+			PublicDefinitions.Add("WITH_AITESTSUITE 1");
 		}
 	}
 }
diff --git a/Engine/Source/Developer/AITestSuite/Private/AITestSuite.cpp b/Engine/Source/Developer/AITestSuite/Private/AITestSuite.cpp
--- a/Engine/Source/Developer/AITestSuite/Private/AITestSuite.cpp
+++ b/Engine/Source/Developer/AITestSuite/Private/AITestSuite.cpp
@@ -1,10 +1,7 @@
 // Copyright Epic Games, Inc. All Rights Reserved.
 
 #include "AITestSuite.h"
-#include "AITestsCommon.h"
 
-DEFINE_LOG_CATEGORY(LogAITestSuite);
-DEFINE_LOG_CATEGORY(LogBehaviorTreeTest);
 
 class FAITestSuite : public IAITestSuite
 {
diff --git a/Engine/Source/Developer/AITestSuite/Private/Tests/AITestsCommon.cpp b/Engine/Source/Developer/AITestSuite/Private/Tests/AITestsCommon.cpp
--- a/Engine/Source/Developer/AITestSuite/Private/Tests/AITestsCommon.cpp
+++ b/Engine/Source/Developer/AITestSuite/Private/Tests/AITestsCommon.cpp
@@ -4,6 +4,15 @@
 #include "EngineGlobals.h"
 #include "Engine/Engine.h"
 
+DEFINE_LOG_CATEGORY(LogAITestSuite);
+DEFINE_LOG_CATEGORY(LogBehaviorTreeTest);
+DEFINE_LOG_CATEGORY(LogTestableEnsures);
+
+
+namespace UE::AITestSuite
+{
+	int32 TestsInProgress = 0;
+} // namespace UE::AITestSuite
 
 namespace FAITestHelpers
 {
diff --git a/Engine/Source/Developer/AITestSuite/Public/AITestsCommon.h b/Engine/Source/Developer/AITestSuite/Public/AITestsCommon.h
--- a/Engine/Source/Developer/AITestSuite/Public/AITestsCommon.h
+++ b/Engine/Source/Developer/AITestSuite/Public/AITestsCommon.h
@@ -2,13 +2,13 @@
 
 #pragma once
 
-#include "CoreMinimal.h"
 #include "UObject/Package.h"
 #include "UObject/UObjectGlobals.h"
+#include "CoreGlobals.h"
 #include "Misc/AutomationTest.h"
 #include "TestLogger.h"
 #include "Engine/EngineBaseTypes.h"
-
+#include "TestableEnsures.h"
 
 DECLARE_LOG_CATEGORY_EXTERN(LogAITestSuite, Log, All);
 DECLARE_LOG_CATEGORY_EXTERN(LogBehaviorTreeTest, Log, All);
@@ -26,7 +26,6 @@ private:
 	bool bAlreadyRun;
 };
 
-
 namespace FAITestHelpers
 {
 	AITESTSUITE_API UWorld* GetWorld();
diff --git a/Engine/Source/Developer/AITestSuite/Public/TestableEnsures.h b/Engine/Source/Developer/AITestSuite/Public/TestableEnsures.h
new file mode 100644
--- /dev/null
+++ b/Engine/Source/Developer/AITestSuite/Public/TestableEnsures.h
@@ -0,0 +1,60 @@
+// Copyright Epic Games, Inc. All Rights Reserved.
+
+#pragma once
+
+#include "CoreGlobals.h"
+#include "HAL/Platform.h"
+
+AITESTSUITE_API DECLARE_LOG_CATEGORY_EXTERN(LogTestableEnsures, Log, All);
+
+//-----------------------------------------------------------------------------
+// support for testable ensure and check
+//-----------------------------------------------------------------------------
+namespace UE::AITestSuite
+{
+	AITESTSUITE_API extern int32 TestsInProgress;
+}
+
+#define AITEST_SCOPED_CHECK(PartialTextToExpect, NumOccurrences) \
+	GetTestRunner().AddExpectedError(PartialTextToExpect, EAutomationExpectedErrorFlags::Contains, NumOccurrences); \
+	ON_SCOPE_EXIT \
+	{ \
+		--UE::AITestSuite::TestsInProgress; \
+	}; \
+	++UE::AITestSuite::TestsInProgress \
+
+#define testableEnsureMsgf(InExpression, InFormat, ... ) \
+	(LIKELY(!!(InExpression)) || ([&]() \
+		{ \
+			if (UNLIKELY(UE::AITestSuite::TestsInProgress > 0)) \
+			{ \
+				UE_LOG(LogTestableEnsures, Warning, InFormat, ##__VA_ARGS__); \
+			} \
+			else \
+			{ \
+				ensureMsgf(false, InFormat, ##__VA_ARGS__); \
+			} \
+		return false; \
+		} ()))
+
+#define testableCheckf(InExpression, InFormat, ... ) \
+	if (UNLIKELY(UE::AITestSuite::TestsInProgress > 0)) \
+	{ \
+		UE_CLOG(!(InExpression), LogTestableEnsures, Error, InFormat, ##__VA_ARGS__); \
+		return; \
+	} \
+	else \
+	{ \
+		checkf(InExpression, InFormat, ##__VA_ARGS__); \
+	}
+
+#define testableCheckfReturn(InExpression, ReturnValue, InFormat, ... ) \
+	if (UNLIKELY(UE::AITestSuite::TestsInProgress > 0)) \
+	{ \
+		UE_CLOG(!(InExpression), LogTestableEnsures, Error, InFormat, ##__VA_ARGS__); \
+		return (ReturnValue); \
+	} \
+	else \
+	{ \
+		checkf(InExpression, InFormat, ##__VA_ARGS__); \
+	}
