Fixed the testableCheckfReturn macro to not always execute the return expression while unit testing was running, even if the check being evaluated did not fail

diff --git a/Engine/Source/Developer/AITestSuite/Public/TestableEnsures.h b/Engine/Source/Developer/AITestSuite/Public/TestableEnsures.h
--- a/Engine/Source/Developer/AITestSuite/Public/TestableEnsures.h
+++ b/Engine/Source/Developer/AITestSuite/Public/TestableEnsures.h
@@ -48,11 +48,14 @@ namespace UE::AITestSuite
 		checkf(InExpression, InFormat, ##__VA_ARGS__); \
 	}
 
-#define testableCheckfReturn(InExpression, ReturnValue, InFormat, ... ) \
+#define testableCheckfReturn(InExpression, ReturnExpression, InFormat, ... ) \
 	if (UNLIKELY(UE::AITestSuite::TestsInProgress > 0)) \
 	{ \
-		UE_CLOG(!(InExpression), LogTestableEnsures, Error, InFormat, ##__VA_ARGS__); \
-		return (ReturnValue); \
+		if (!(InExpression)) \
+		{ \
+			UE_LOG(LogTestableEnsures, Error, InFormat, ##__VA_ARGS__); \
+			ReturnExpression; \
+		} \
 	} \
 	else \
 	{ \
