[MockAI_BT] fixed component initialization when tests are run in a game world

diff --git a/Engine/Source/Developer/AITestSuite/Private/MockAI/MockAI_BT.cpp b/Engine/Source/Developer/AITestSuite/Private/MockAI/MockAI_BT.cpp
--- a/Engine/Source/Developer/AITestSuite/Private/MockAI/MockAI_BT.cpp
+++ b/Engine/Source/Developer/AITestSuite/Private/MockAI/MockAI_BT.cpp
@@ -25,7 +25,14 @@ UMockAI_BT::UMockAI_BT(const FObjectInitializer& ObjectInitializer)
 		// We want our component to be ticked manually by the test framework, never by the TickFunction
 		BTComp->PrimaryComponentTick.bCanEverTick = false;
 		BTComp->RegisterComponent();
-		BTComp->InitializeComponent();
+
+		// Test running in a game world will get initialized automatically by
+		// the previous call to RegisterComponent. Otherwise, we initialize it
+		// manually.
+		if (!BTComp->HasBeenInitialized() && BTComp->bWantsInitializeComponent)
+		{
+			BTComp->InitializeComponent();
+		}
 	}
 }
 
