[AITestSuite] Added a console variable, "test.BreakOnTestFail", that if set to true will cause UE_DEBUG_BREAK on test failures, as tested by AITEST macros.

Super handy when fixing tests and/or test-driven-developing your code (as I often do).

#rnx

diff --git a/Engine/Source/Developer/AITestSuite/Private/Tests/AITestsCommon.cpp b/Engine/Source/Developer/AITestSuite/Private/Tests/AITestsCommon.cpp
--- a/Engine/Source/Developer/AITestSuite/Private/Tests/AITestsCommon.cpp
+++ b/Engine/Source/Developer/AITestSuite/Private/Tests/AITestsCommon.cpp
@@ -9,6 +9,23 @@ DEFINE_LOG_CATEGORY(LogBehaviorTreeTest);
 DEFINE_LOG_CATEGORY(LogTestableEnsures);
 
 
+namespace UE::AITestSuite
+{
+	bool bBreakOnTestFail = false;
+	FAutoConsoleVariableRef CVarBreakOnTestFail(TEXT("test.BreakOnTestFail")
+		, bBreakOnTestFail
+		, TEXT("Enables AITestSuite-implemented tests to cause a debug break when tests fail. Used when fixing test offenders or during test-driven development.")
+		, ECVF_Cheat);
+
+	void ConditionallyBreakOnTestFail()
+	{
+		if (bBreakOnTestFail)
+		{
+			UE_DEBUG_BREAK();
+		}
+	}
+}
+
 namespace UE::AITestSuite
 {
 	int32 TestsInProgress = 0;
diff --git a/Engine/Source/Developer/AITestSuite/Public/AITestsCommon.h b/Engine/Source/Developer/AITestSuite/Public/AITestsCommon.h
--- a/Engine/Source/Developer/AITestSuite/Public/AITestsCommon.h
+++ b/Engine/Source/Developer/AITestSuite/Public/AITestsCommon.h
@@ -3,8 +3,10 @@
 #pragma once
 
 #include "UObject/Package.h"
+#if UE_ENABLE_INCLUDE_ORDER_DEPRECATED_IN_5_6
 #include "UObject/UObjectGlobals.h"
 #include "CoreGlobals.h"
+#endif // UE_ENABLE_INCLUDE_ORDER_DEPRECATED_IN_5_6
 #include "Misc/AutomationTest.h"
 #include "TestLogger.h"
 #include "Engine/EngineBaseTypes.h"
@@ -15,6 +17,12 @@ DECLARE_LOG_CATEGORY_EXTERN(LogBehaviorTreeTest, Log, All);
 
 DEFINE_LATENT_AUTOMATION_COMMAND_ONE_PARAMETER(FAITestCommand_WaitSeconds, float, Duration);
 
+namespace UE::AITestSuite
+{
+	/** Will cause UE_DEBUG_BREAK if `test.BreakOnTestFail` console variable is true */
+	AITESTSUITE_API void ConditionallyBreakOnTestFail(); 
+}
+
 class FAITestCommand_WaitOneTick : public IAutomationLatentCommand
 {
 public: 
@@ -29,7 +37,7 @@ private:
 namespace FAITestHelpers
 {
 	AITESTSUITE_API UWorld* GetWorld();
-	static const float TickInterval = 1.f / 30;
+	static constexpr float TickInterval = 1.f / 30;
 
 	void UpdateFrameCounter();
 	uint64 FramesCounter();
@@ -75,7 +83,7 @@ public:
 	virtual void VerifyLatentResults() {}
 	/** @return false to indicate an issue with test execution. Will signal to automation framework this test instance failed. */
 	virtual bool InstantTest() { return false;}
-	// it's essential that overriding functions call the super-implementation. Otherwise the check in ~FAITestBase will fail.
+	// it's essential that overriding functions call the super-implementation. Otherwise, the check in ~FAITestBase will fail.
 	virtual void TearDown();
 };
 
@@ -202,6 +210,7 @@ struct FAITest_SimpleComponentBasedTest : public FAITestBase
 #define __AITEST_IMPL(What, Value, Test, RetVal)\
 	if (!GetTestRunner().Test(What, Value))\
 	{\
+		UE::AITestSuite::ConditionallyBreakOnTestFail(); \
 		return RetVal;\
 	}
 	
@@ -209,7 +218,7 @@ struct FAITest_SimpleComponentBasedTest : public FAITestBase
 #define AITEST_FALSE(What, Value) __AITEST_IMPL(What, Value, TestFalse, false)
 #define AITEST_NULL(What, Pointer) __AITEST_IMPL(What, Pointer, TestNull, false)
 #define AITEST_NOT_NULL(What, Pointer) \
-	__AITEST_IMPL(What, Pointer, TestNotNull, false) \
+	__AITEST_IMPL(What, Pointer, TestNotNull, false); \
 	CA_ASSUME(Pointer)
 
 namespace FTestHelpers
@@ -242,6 +251,7 @@ namespace FTestHelpers
 #define __AITEST_HELPER_IMPL(What, Actual, Expected, Test, RetVal)\
 	if (!FTestHelpers::Test(What, Actual, Expected, GetTestRunner()))\
 	{\
+		UE::AITestSuite::ConditionallyBreakOnTestFail(); \
 		return RetVal;\
 	}
 
