[AITestingSuite] Fixed test.BreakOnTestFail so that it's possible to set it via the DefaultEngine.ini

#rnx

diff --git a/Engine/Source/Developer/AITestSuite/Private/Tests/AITestsCommon.cpp b/Engine/Source/Developer/AITestSuite/Private/Tests/AITestsCommon.cpp
--- a/Engine/Source/Developer/AITestSuite/Private/Tests/AITestsCommon.cpp
+++ b/Engine/Source/Developer/AITestSuite/Private/Tests/AITestsCommon.cpp
@@ -15,7 +15,7 @@ namespace UE::AITestSuite
 	FAutoConsoleVariableRef CVarBreakOnTestFail(TEXT("test.BreakOnTestFail")
 		, bBreakOnTestFail
 		, TEXT("Enables AITestSuite-implemented tests to cause a debug break when tests fail. Used when fixing test offenders or during test-driven development.")
-		, ECVF_Cheat);
+		, ECVF_Default);
 
 	void ConditionallyBreakOnTestFail()
 	{
