[AITestingSuite] Added a AITEST_INFO macro for annotating test logs.

#rnx

diff --git a/Engine/Source/Developer/AITestSuite/Public/AITestsCommon.h b/Engine/Source/Developer/AITestSuite/Public/AITestsCommon.h
--- a/Engine/Source/Developer/AITestSuite/Public/AITestsCommon.h
+++ b/Engine/Source/Developer/AITestSuite/Public/AITestsCommon.h
@@ -246,6 +246,11 @@ namespace FTestHelpers
 	{
 		return This.TestNotEqual(*Description, reinterpret_cast<uint64>(Expression), reinterpret_cast<uint64>(Expected));
 	}
+
+	inline void AddInfo(const FString& Info, FAutomationTestBase& This)
+	{
+		This.AddInfo(*Info);
+	}
 }
 
 #define __AITEST_HELPER_IMPL(What, Actual, Expected, Test, RetVal)\
@@ -258,6 +263,8 @@ namespace FTestHelpers
 #define AITEST_EQUAL(What, Actual, Expected) __AITEST_HELPER_IMPL(What, Actual, Expected, TestEqual, false)
 #define AITEST_NOT_EQUAL(What, Actual, Expected) __AITEST_HELPER_IMPL(What, Actual, Expected, TestNotEqual, false)
 
+#define AITEST_INFO(InfoString) FTestHelpers::AddInfo(InfoString, GetTestRunner())
+
 
 //----------------------------------------------------------------------//
 // state testing macros, valid in FTestAIBase (and subclasses') methods. 
