[BehaviorTree] ApplySearchUpdatesImpl now applies the active instance index stored in pending update to make sure methods calls use the right instance (i.e., decorators will get memory from the right buffer)
- added unit test to abort a simple parallel node from a subtree

diff --git a/Engine/Source/Developer/AITestSuite/Classes/BehaviorTree/TestBTDecorator_NodeMemoryValidator.h b/Engine/Source/Developer/AITestSuite/Classes/BehaviorTree/TestBTDecorator_NodeMemoryValidator.h
new file mode 100644
--- /dev/null
+++ b/Engine/Source/Developer/AITestSuite/Classes/BehaviorTree/TestBTDecorator_NodeMemoryValidator.h
@@ -0,0 +1,47 @@
+// Copyright Epic Games, Inc. All Rights Reserved.
+
+#pragma once 
+
+#include "BehaviorTree/BTDecorator.h"
+#include "TestBTDecorator_NodeMemoryValidator.generated.h"
+
+struct FBTNodeMemoryValidatorDecoratorMemory
+{
+	static constexpr uint32 TestValue = 0xABBAABBA;
+	uint32 Dummy;
+};
+
+UCLASS(meta=(HiddenNode))
+class UTestBTDecorator_NodeMemoryValidator : public UBTDecorator
+{
+	GENERATED_UCLASS_BODY()
+
+protected:
+	virtual void OnBecomeRelevant(UBehaviorTreeComponent& OwnerComp, uint8* NodeMemory) override;
+	virtual void OnCeaseRelevant(UBehaviorTreeComponent& OwnerComp, uint8* NodeMemory) override;
+	virtual void TickNode(UBehaviorTreeComponent& OwnerComp, uint8* NodeMemory, float DeltaSeconds) override;
+	virtual void OnNodeActivation(FBehaviorTreeSearchData& SearchData) override;
+	virtual void OnNodeDeactivation(FBehaviorTreeSearchData& SearchData, EBTNodeResult::Type NodeResult) override;
+	virtual void OnNodeProcessed(FBehaviorTreeSearchData& SearchData, EBTNodeResult::Type& NodeResult) override;
+	virtual bool CalculateRawConditionValue(UBehaviorTreeComponent& OwnerComp, uint8* NodeMemory) const override;
+
+	virtual void InitializeMemory(UBehaviorTreeComponent& OwnerComp, uint8* NodeMemory, EBTMemoryInit::Type InitType) const override;
+	virtual void CleanupMemory(UBehaviorTreeComponent& OwnerComp, uint8* NodeMemory, EBTMemoryClear::Type CleanupType) const override;
+	virtual uint16 GetInstanceMemorySize() const override;
+
+private:
+
+	template<typename T>
+	T* GetCheckedNodeMemory(FBehaviorTreeSearchData& SearchData) const
+	{
+		checkf(SearchData.OwnerComp.FindInstanceContainingNode(this) == SearchData.OwnerComp.GetActiveInstanceIdx(), TEXT("Accessing memory from the wrong instance"));
+		return GetNodeMemory<T>(SearchData);
+	}
+
+	template<typename T>
+	const T* GetCheckedNodeMemory(const FBehaviorTreeSearchData& SearchData) const
+	{
+		checkf(SearchData.OwnerComp.FindInstanceContainingNode(this) == SearchData.OwnerComp.GetActiveInstanceIdx(), TEXT("Accessing memory from the wrong instance"));
+		return GetNodeMemory<T>(SearchData);
+	}
+};
diff --git a/Engine/Source/Developer/AITestSuite/Classes/BehaviorTree/TestBTTask_RunBehavior.h b/Engine/Source/Developer/AITestSuite/Classes/BehaviorTree/TestBTTask_RunBehavior.h
new file mode 100644
--- /dev/null
+++ b/Engine/Source/Developer/AITestSuite/Classes/BehaviorTree/TestBTTask_RunBehavior.h
@@ -0,0 +1,22 @@
+// Copyright Epic Games, Inc. All Rights Reserved.
+
+#pragma once
+
+#include "BehaviorTree/Tasks/BTTask_RunBehavior.h"
+#include "TestBTTask_RunBehavior.generated.h"
+
+UCLASS(MinimalAPI)
+class UTestBTTask_RunBehavior : public UBTTask_RunBehavior
+{
+	GENERATED_UCLASS_BODY()
+	
+	/**
+	 * Sets the subtree asset to run.
+	 * @note This method has been added for test purpose only
+	 * and should be called only before the task gets activated.
+	 */
+	void SetSubtreeAsset(UBehaviorTree* Asset)
+	{
+		BehaviorAsset = Asset;
+	}
+};
diff --git a/Engine/Source/Developer/AITestSuite/Private/BehaviorTree/TestBTDecorator_NodeMemoryValidator.cpp b/Engine/Source/Developer/AITestSuite/Private/BehaviorTree/TestBTDecorator_NodeMemoryValidator.cpp
new file mode 100644
--- /dev/null
+++ b/Engine/Source/Developer/AITestSuite/Private/BehaviorTree/TestBTDecorator_NodeMemoryValidator.cpp
@@ -0,0 +1,80 @@
+// Copyright Epic Games, Inc. All Rights Reserved.
+
+#include "BehaviorTree/TestBTDecorator_NodeMemoryValidator.h"
+
+#include UE_INLINE_GENERATED_CPP_BY_NAME(TestBTDecorator_NodeMemoryValidator)
+
+UTestBTDecorator_NodeMemoryValidator::UTestBTDecorator_NodeMemoryValidator(const FObjectInitializer& ObjectInitializer) : Super(ObjectInitializer)
+{
+	NodeName = TEXT("NodeMemoryValidator");
+
+	bAllowAbortNone = false;
+	bAllowAbortLowerPri = false;
+	bAllowAbortChildNodes = false;
+
+	INIT_DECORATOR_NODE_NOTIFY_FLAGS();
+}
+
+void UTestBTDecorator_NodeMemoryValidator::OnBecomeRelevant(UBehaviorTreeComponent& OwnerComp, uint8* NodeMemory)
+{
+	const FBTNodeMemoryValidatorDecoratorMemory* DecoratorMemory = CastInstanceNodeMemory<FBTNodeMemoryValidatorDecoratorMemory>(NodeMemory);
+	check(DecoratorMemory->Dummy == FBTNodeMemoryValidatorDecoratorMemory::TestValue);
+}
+
+void UTestBTDecorator_NodeMemoryValidator::OnCeaseRelevant(UBehaviorTreeComponent& OwnerComp, uint8* NodeMemory)
+{
+	const FBTNodeMemoryValidatorDecoratorMemory* DecoratorMemory = CastInstanceNodeMemory<FBTNodeMemoryValidatorDecoratorMemory>(NodeMemory);
+	check(DecoratorMemory->Dummy == FBTNodeMemoryValidatorDecoratorMemory::TestValue);
+}
+
+void UTestBTDecorator_NodeMemoryValidator::TickNode(UBehaviorTreeComponent& OwnerComp, uint8* NodeMemory, float DeltaSeconds)
+{
+	const FBTNodeMemoryValidatorDecoratorMemory* DecoratorMemory = CastInstanceNodeMemory<FBTNodeMemoryValidatorDecoratorMemory>(NodeMemory);
+	check(DecoratorMemory->Dummy == FBTNodeMemoryValidatorDecoratorMemory::TestValue);
+}
+
+void UTestBTDecorator_NodeMemoryValidator::OnNodeActivation(FBehaviorTreeSearchData& SearchData)
+{
+	const FBTNodeMemoryValidatorDecoratorMemory* DecoratorMemory = GetCheckedNodeMemory<FBTNodeMemoryValidatorDecoratorMemory>(SearchData);
+	check(DecoratorMemory->Dummy == FBTNodeMemoryValidatorDecoratorMemory::TestValue);
+}
+
+void UTestBTDecorator_NodeMemoryValidator::OnNodeDeactivation(FBehaviorTreeSearchData& SearchData, EBTNodeResult::Type NodeResult)
+{
+	const FBTNodeMemoryValidatorDecoratorMemory* DecoratorMemory = GetCheckedNodeMemory<FBTNodeMemoryValidatorDecoratorMemory>(SearchData);
+	check(DecoratorMemory->Dummy == FBTNodeMemoryValidatorDecoratorMemory::TestValue);
+}
+
+void UTestBTDecorator_NodeMemoryValidator::OnNodeProcessed(FBehaviorTreeSearchData& SearchData, EBTNodeResult::Type& NodeResult)
+{
+	const FBTNodeMemoryValidatorDecoratorMemory* DecoratorMemory = GetCheckedNodeMemory<FBTNodeMemoryValidatorDecoratorMemory>(SearchData);
+	check(DecoratorMemory->Dummy == FBTNodeMemoryValidatorDecoratorMemory::TestValue);
+}
+
+bool UTestBTDecorator_NodeMemoryValidator::CalculateRawConditionValue(UBehaviorTreeComponent& OwnerComp, uint8* NodeMemory) const
+{
+	const FBTNodeMemoryValidatorDecoratorMemory* DecoratorMemory = CastInstanceNodeMemory<FBTNodeMemoryValidatorDecoratorMemory>(NodeMemory);
+	check(DecoratorMemory->Dummy == FBTNodeMemoryValidatorDecoratorMemory::TestValue);
+	return Super::CalculateRawConditionValue(OwnerComp, NodeMemory);
+}
+
+uint16 UTestBTDecorator_NodeMemoryValidator::GetInstanceMemorySize() const
+{
+	return sizeof(FBTNodeMemoryValidatorDecoratorMemory);
+}
+
+void UTestBTDecorator_NodeMemoryValidator::InitializeMemory(UBehaviorTreeComponent& OwnerComp, uint8* NodeMemory, EBTMemoryInit::Type InitType) const
+{
+	FBTNodeMemoryValidatorDecoratorMemory* DecoratorMemory = InitializeNodeMemory<FBTNodeMemoryValidatorDecoratorMemory>(NodeMemory, InitType);
+	if (InitType == EBTMemoryInit::Initialize)
+	{
+		DecoratorMemory->Dummy = FBTNodeMemoryValidatorDecoratorMemory::TestValue;
+	}
+}
+
+void UTestBTDecorator_NodeMemoryValidator::CleanupMemory(UBehaviorTreeComponent& OwnerComp, uint8* NodeMemory, EBTMemoryClear::Type CleanupType) const
+{
+	const FBTNodeMemoryValidatorDecoratorMemory* DecoratorMemory = CastInstanceNodeMemory<FBTNodeMemoryValidatorDecoratorMemory>(NodeMemory);
+	check(DecoratorMemory->Dummy == FBTNodeMemoryValidatorDecoratorMemory::TestValue);
+	CleanupNodeMemory<FBTNodeMemoryValidatorDecoratorMemory>(NodeMemory, CleanupType);
+}
diff --git a/Engine/Source/Developer/AITestSuite/Private/BehaviorTree/TestBTTask_RunBehavior.cpp b/Engine/Source/Developer/AITestSuite/Private/BehaviorTree/TestBTTask_RunBehavior.cpp
new file mode 100644
--- /dev/null
+++ b/Engine/Source/Developer/AITestSuite/Private/BehaviorTree/TestBTTask_RunBehavior.cpp
@@ -0,0 +1,10 @@
+// Copyright Epic Games, Inc. All Rights Reserved.
+
+#include "BehaviorTree/TestBTTask_RunBehavior.h"
+
+#include UE_INLINE_GENERATED_CPP_BY_NAME(TestBTTask_RunBehavior)
+
+UTestBTTask_RunBehavior::UTestBTTask_RunBehavior(const FObjectInitializer& ObjectInitializer) : Super(ObjectInitializer)
+{
+	NodeName = TEXT("Test Run Behavior");
+}
diff --git a/Engine/Source/Developer/AITestSuite/Private/Tests/BTTest.cpp b/Engine/Source/Developer/AITestSuite/Private/Tests/BTTest.cpp
--- a/Engine/Source/Developer/AITestSuite/Private/Tests/BTTest.cpp
+++ b/Engine/Source/Developer/AITestSuite/Private/Tests/BTTest.cpp
@@ -3,10 +3,12 @@
 #include "AITestsCommon.h"
 #include "BTBuilder.h"
 #include "BehaviorTree/BTCompositeNode.h"
+#include "BehaviorTree/Decorators/BTDecorator_Cooldown.h"
 #include "BehaviorTree/Decorators/BTDecorator_ForceSuccess.h"
 #include "BehaviorTree/Decorators/BTDecorator_TimeLimit.h"
 #include "BehaviorTree/TestBTDecorator_CantExecute.h"
 #include "BehaviorTree/TestBTTask_TimerBasedLatent.h"
+#include "BehaviorTree/TestBTTask_RunBehavior.h"
 #include "Engine/World.h"
 #include "MockAI_BT.h"
 #include "TimerManager.h"
@@ -4442,6 +4444,59 @@ struct FAITest_BTTestTickAuxNodesDeactivateNodes : public FAITest_SimpleBT
 		ExpectedResult.Add(SecondaryTask);
 	}
 };
-IMPLEMENT_AI_LATENT_TEST(FAITest_BTTestTickAuxNodesDeactivateNodes, "System.AI.Behavior Trees.AITest_TickAuxNodesDeactivateNodes: Ticking Aux Nodes Deactivates Nodes")
+IMPLEMENT_AI_LATENT_TEST(FAITest_BTTestTickAuxNodesDeactivateNodes, "System.AI.Behavior Trees.Ticking aux nodes deactivates nodes")
+
+struct FAITest_BTTestAbortNestedSimpleParallel : public FAITest_SimpleBT
+{
+	FAITest_BTTestAbortNestedSimpleParallel()
+	{
+		enum
+		{
+			ChildFirstTaskExecute = 1,
+			ChildSecondTaskExecute,
+			ChildThirdTaskExecute
+		};
+
+		UBehaviorTree* ChildBTAsset = &FBTBuilder::CreateBehaviorTree();
+		AddAutoDestroyObject(*ChildBTAsset);
+
+		// Setup subtree
+		{
+			UBTCompositeNode& RootNode = FBTBuilder::AddSequence(*ChildBTAsset);
+			{
+				UBTCompositeNode& CompNode = FBTBuilder::AddParallel(RootNode, EBTParallelMode::AbortBackground);
+				{
+					FBTBuilder::AddTaskLogFinish(/*ParentNode*/ CompNode, /*LogIndex*/ ChildFirstTaskExecute, /*FinishIndex*/ INDEX_NONE, EBTNodeResult::Succeeded, /*ExecutionTick*/ 500);
+					FBTBuilder::WithDecorator<UBTDecorator_Cooldown>(CompNode);
+
+					UBTCompositeNode& CompNode2 = FBTBuilder::AddSequence(CompNode);
+					{
+						FBTBuilder::AddTaskLogFinish(/*ParentNode*/ CompNode2, /*LogIndex*/ ChildSecondTaskExecute, /*FinishIndex*/ INDEX_NONE, EBTNodeResult::Succeeded, /*ExecutionTick*/ 5);
+						FBTBuilder::WithDecorator<UBTDecorator_Cooldown>(CompNode2);
+
+						FBTBuilder::AddTaskLogFinish(/*ParentNode*/ CompNode2, /*LogIndex*/ ChildThirdTaskExecute, /*FinishIndex*/ INDEX_NONE, EBTNodeResult::Succeeded, /*ExecutionTick*/ 5);
+					}
+				}
+			}
+		}
+
+		// Setup main tree
+		{
+			UBTCompositeNode& RootNode = FBTBuilder::AddSequence(*BTAsset);
+			{
+				constexpr float NumTicks = 4.f;
+				UBTCompositeNode& CompNode = FBTBuilder::AddSelector(RootNode);
+				FBTBuilder::WithDecorator<UBTDecorator_TimeLimit>(RootNode).TimeLimit = NumTicks * FAITestHelpers::TickInterval;
+				{
+					FBTBuilder::AddTask<UTestBTTask_RunBehavior>(CompNode).SetSubtreeAsset(ChildBTAsset);
+				}
+			}
+		}
+
+		ExpectedResult.Add(ChildFirstTaskExecute);
+		ExpectedResult.Add(ChildSecondTaskExecute);
+	}
+};
+IMPLEMENT_AI_LATENT_TEST(FAITest_BTTestAbortNestedSimpleParallel, "System.AI.Behavior Trees.Abort: nested simple parallel")
 
 #undef LOCTEXT_NAMESPACE
\ No newline at end of file
diff --git a/Engine/Source/Developer/AITestSuite/Public/BTBuilder.h b/Engine/Source/Developer/AITestSuite/Public/BTBuilder.h
--- a/Engine/Source/Developer/AITestSuite/Public/BTBuilder.h
+++ b/Engine/Source/Developer/AITestSuite/Public/BTBuilder.h
@@ -18,6 +18,7 @@
 #include "BehaviorTree/TestBTDecorator_Blackboard.h"
 #include "BehaviorTree/TestBTDecorator_DelayedAbort.h"
 #include "BehaviorTree/TestBTDecorator_Blueprint.h"
+#include "BehaviorTree/TestBTDecorator_NodeMemoryValidator.h"
 #include "BehaviorTree/TestBTService_Log.h"
 #include "BehaviorTree/TestBTService_BTStopAction.h"
 #include "BehaviorTree/TestBTTask_LatentWithFlags.h"
@@ -276,6 +277,9 @@ struct FBTBuilder
 		T* DecoratorOb = NewObject<T>(ParentNode.GetTreeAsset());
 		ParentNode.Children.Last().Decorators.Add(DecoratorOb);
 
+		// Push additional memory validation decorator for extra coverage
+		ParentNode.Children.Last().Decorators.Add(NewObject<UTestBTDecorator_NodeMemoryValidator>(ParentNode.GetTreeAsset()));
+
 		return *DecoratorOb;
 	}
 
diff --git a/Engine/Source/Runtime/AIModule/Private/BehaviorTree/BehaviorTreeComponent.cpp b/Engine/Source/Runtime/AIModule/Private/BehaviorTree/BehaviorTreeComponent.cpp
--- a/Engine/Source/Runtime/AIModule/Private/BehaviorTree/BehaviorTreeComponent.cpp
+++ b/Engine/Source/Runtime/AIModule/Private/BehaviorTree/BehaviorTreeComponent.cpp
@@ -1470,6 +1470,10 @@ void UBehaviorTreeComponent::ApplySearchUpdatesImpl(const TArray<FBehaviorTreeSe
 			continue;
 		}
 
+		// Push the stored instance index as the active one so calls from the following processing 
+		// will get the right instance when calling GetActiveInstanceIdx();
+		TGuardValue GuardActiveInstanceIndex(ActiveInstanceIdx, UpdateInfo.InstanceIndex);
+
 		FBehaviorTreeInstance& UpdateInstance = InstanceStack[UpdateInfo.InstanceIndex];
 		int32 ParallelTaskIdx = INDEX_NONE;
 		bool bIsComponentActive = false;
