[BehaviorTree] fixed BTDecorator_ForceSuccess when used used on aborted RunBehavior nodes

diff --git a/Engine/Source/Developer/AITestSuite/Private/Tests/BTTest.cpp b/Engine/Source/Developer/AITestSuite/Private/Tests/BTTest.cpp
--- a/Engine/Source/Developer/AITestSuite/Private/Tests/BTTest.cpp
+++ b/Engine/Source/Developer/AITestSuite/Private/Tests/BTTest.cpp
@@ -506,6 +506,50 @@
 };
 IMPLEMENT_AI_LATENT_TEST(FAITest_BTDecoratorForceSuccessCompositeFailedChild, "System.AI.Behavior Trees.Decorator: force success composite failed child task")
 
+struct FAITest_BTDecoratorForceSuccessCompositeFailedChildFromSubTree : public FAITest_SimpleBT
+{
+	FAITest_BTDecoratorForceSuccessCompositeFailedChildFromSubTree()
+	{
+		enum
+		{
+			ChildTask1Execute = 1,
+			ChildTask2Execute,
+			Task1Execute,
+			Task2Execute,
+		};
+
+		// Subtree that will trigger the abort condition (i.e, Bool1 set) of the composite node in the parent tree
+		UBehaviorTree* ChildAsset1 = &FBTBuilder::CreateBehaviorTree(*BTAsset);
+		if (ChildAsset1)
+		{
+			AddAutoDestroyObject(*ChildAsset1);
+			UBTCompositeNode& CompNode = FBTBuilder::AddSequence(*ChildAsset1);
+			{
+				FBTBuilder::AddTask(CompNode, ChildTask1Execute, EBTNodeResult::Succeeded);
+				FBTBuilder::AddTaskFlagChange(CompNode, true, EBTNodeResult::InProgress, TEXT("Bool1"));
+				FBTBuilder::AddTask(CompNode, ChildTask2Execute, EBTNodeResult::Succeeded);
+			}
+		}
+
+		// Main tree with a composite that aborts on key change but should succeed with ForceSuccess decorator
+		UBTCompositeNode& MainSequence = FBTBuilder::AddSequence(*BTAsset);
+		{
+			FBTBuilder::AddTask(MainSequence, Task1Execute, EBTNodeResult::Succeeded);
+			FBTBuilder::AddTaskSubtree(MainSequence, ChildAsset1);
+			{
+				FBTBuilder::WithDecoratorBlackboard(MainSequence, EBasicKeyOperation::NotSet, EBTFlowAbortMode::Self, TEXT("Bool1"));
+				FBTBuilder::WithDecorator<UBTDecorator_ForceSuccess>(MainSequence);
+			}
+			FBTBuilder::AddTask(MainSequence, Task2Execute, EBTNodeResult::Succeeded);
+		}
+
+		ExpectedResult.Add(Task1Execute);
+		ExpectedResult.Add(ChildTask1Execute);
+		ExpectedResult.Add(Task2Execute);
+	}
+};
+IMPLEMENT_AI_LATENT_TEST(FAITest_BTDecoratorForceSuccessCompositeFailedChildFromSubTree, "System.AI.Behavior Trees.Decorator: force success composite failed from subtree task")
+
 struct FAITest_BTDecoratorForceSuccessTaskFailedSibling : public FAITest_SimpleBT
 {
 	FAITest_BTDecoratorForceSuccessTaskFailedSibling()
@@ -902,7 +946,6 @@
 };
 IMPLEMENT_AI_LATENT_TEST(FAITest_BTAbortingDuringServiceTick, "System.AI.Behavior Trees.Abort: during service tick")
 
-
 struct FAITest_BTRestartingTheAbortedBranch : public FAITest_SimpleBT
 {
 	FAITest_BTRestartingTheAbortedBranch()
@@ -936,7 +979,6 @@
 };
 IMPLEMENT_AI_LATENT_TEST(FAITest_BTRestartingTheAbortedBranch, "System.AI.Behavior Trees.Toggle: restart the aborted branch")
 
-
 struct FAITest_BTAbortingTheRestarted : public FAITest_SimpleBT
 {
 	FAITest_BTAbortingTheRestarted()
@@ -965,7 +1007,6 @@
 };
 IMPLEMENT_AI_LATENT_TEST(FAITest_BTAbortingTheRestarted, "System.AI.Behavior Trees.Toggle: reabort the restarted branch")
 
-
 struct FAITest_BTRestartingTheReaborted : public FAITest_SimpleBT
 {
 	FAITest_BTRestartingTheReaborted()
@@ -999,7 +1040,6 @@
 };
 IMPLEMENT_AI_LATENT_TEST(FAITest_BTRestartingTheReaborted, "System.AI.Behavior Trees.Toggle: restart the reaborted branch")
 
-
 struct FAITest_BTRestartingTheAbortedBranchDuringServiceTick : public FAITest_SimpleBT
 {
 	FAITest_BTRestartingTheAbortedBranchDuringServiceTick()
@@ -1040,7 +1080,6 @@
 };
 IMPLEMENT_AI_LATENT_TEST(FAITest_BTRestartingTheAbortedBranchDuringServiceTick, "System.AI.Behavior Trees.Toggle: restart the aborted branch during service tick")
 
-
 struct FAITest_BTRestartTheAbortedBranchDuringServiceCeaseRelevant : public FAITest_SimpleBT
 {
 	FAITest_BTRestartTheAbortedBranchDuringServiceCeaseRelevant()
@@ -1581,7 +1620,6 @@
 };
 IMPLEMENT_AI_LATENT_TEST(FAITest_BTSwitchingHigherPrioDuringTaskServiceTick, "System.AI.Behavior Trees.Switch: higher priority during task service tick")
 
-
 struct FAITest_BTSwitchingHigherPrioDuringTaskServiceBecomeRelevant : public FAITest_SimpleBT
 {
 	FAITest_BTSwitchingHigherPrioDuringTaskServiceBecomeRelevant()
@@ -2204,7 +2242,6 @@
 };
 IMPLEMENT_AI_LATENT_TEST(FAITest_BTResumingDuringLatentAbort, "System.AI.Behavior Trees.Resume: during latent task abort")
 
-
 struct FAITest_BTAbortDuringInstantAbort : public FAITest_SimpleBT
 {
 	FAITest_BTAbortDuringInstantAbort()
@@ -2651,7 +2688,6 @@
 };
 IMPLEMENT_AI_LATENT_TEST(FAITest_BTSubtreeAbortOutToLowerPrio, "System.AI.Behavior Trees.Subtree: abort out lower prio")
 
-
 /* This unittest come from UDN case 00484061 */
 struct FAITest_BTSubtreeAbortOutToLowerPrio2 : public FAITest_SimpleBT
 {
@@ -2785,7 +2821,6 @@
 };
 IMPLEMENT_AI_LATENT_TEST(FAITest_BTServiceAbortingTask, "System.AI.Behavior Trees.Service: abort task")
 
-
 struct FAITest_BTPostponeDuringSearch : public FAITest_SimpleBT
 {
 	FAITest_BTPostponeDuringSearch()
@@ -2822,7 +2857,6 @@
 };
 IMPLEMENT_AI_LATENT_TEST(FAITest_BTPostponeDuringSearch, "System.AI.Behavior Trees.Service: Postpone during search (high pri)")
 
-
 struct FAITest_BTPostponeDuringSearch2 : public FAITest_SimpleBT
 {
 	FAITest_BTPostponeDuringSearch2()
@@ -2867,6 +2901,139 @@
 };
 IMPLEMENT_AI_LATENT_TEST(FAITest_BTPostponeDuringSearch2, "System.AI.Behavior Trees.Service: Postpone during search (lower pri)")
 
+// This test is temporarily turned off until we enable CVarApplyAuxNodesFromFailedSearches again
+//struct FAITest_BTTestOberserverAbortLowerPriorityTwoDeep: public FAITest_SimpleBT
+//{
+//	FAITest_BTTestOberserverAbortLowerPriorityTwoDeep()
+//	{
+//		enum
+//		{
+//			MainTaskExecute = 1,
+//			InteruptingTaskExecute,
+//		};
+//
+//		UBTCompositeNode& RootNode = FBTBuilder::AddSelector(*BTAsset);
+//		UBTCompositeNode& CompNode = FBTBuilder::AddSelector(RootNode);
+//
+//		{
+//			UBTCompositeNode& Comp1Node = FBTBuilder::AddSelector(CompNode);
+//			FBTBuilder::WithDecoratorBlackboard(CompNode, EArithmeticKeyOperation::Equal, 1, EBTFlowAbortMode::LowerPriority, EBTBlackboardRestart::ValueChange, TEXT("Int"));
+//
+//			UBTCompositeNode& Comp2Node = FBTBuilder::AddSelector(Comp1Node);
+//			FBTBuilder::WithDecoratorBlackboard(Comp1Node, EArithmeticKeyOperation::Equal, 1, EBTFlowAbortMode::LowerPriority, EBTBlackboardRestart::ValueChange, TEXT("Int2"));
+//
+//			FBTBuilder::AddTask(Comp2Node, InteruptingTaskExecute, EBTNodeResult::Succeeded);
+//		}
+//
+//		{
+//			FBTBuilder::AddTaskValuesChangedWithLogs(/*ParentNode*/ CompNode, /*LogIndex*/ MainTaskExecute, /*NodeResult*/ EBTNodeResult::Succeeded, /*Value1*/1, /*Value2*/ 1, /*IntKeyName2*/ TEXT("Int"), /* IntKeyName2 */ TEXT("Int2"), /*ExecutionTicks1*/ 5, /*ExecutionTicks2*/ 5);
+//		}
+//
+//		ExpectedResult.Add(MainTaskExecute);
+//		ExpectedResult.Add(InteruptingTaskExecute);
+//	}
+//};
+//IMPLEMENT_AI_LATENT_TEST(FAITest_BTTestOberserverAbortLowerPriorityTwoDeep, "System.AI.Behavior Trees.OberserverAbortLowerPriorityTwoDeep: Two nodes in same branch with observer aborts lower priority")
+
+struct FAITest_BTTestTickAuxNodesDeactivateNodes : public FAITest_SimpleBT
+{
+	FAITest_BTTestTickAuxNodesDeactivateNodes()
+	{
+		enum
+		{
+			MainTaskExecute = 1,
+			StartAbortMainTask,
+			SecondaryTask,
+			IncorrectTaskExecute,
+		};
+
+		UBTCompositeNode& RootNode = FBTBuilder::AddSelector(*BTAsset);
+		FBTBuilder::WithServiceLog(/*ParentNode*/ RootNode, /*ActivationIndex*/ INDEX_NONE, /*DeactivationIndex*/ INDEX_NONE, /*TickIndex*/ INDEX_NONE, /*TickBoolKeyName*/ TEXT("Bool1"), /*bCallTickOnSearchStart*/ false, /*BecomeRelevantBoolKeyName*/ NAME_None, /*BecomeRelevantBoolKeyName*/ NAME_None, /*bToggleValue*/ false, /*TicksDelaySetKeyNameTick*/ 4);
+		FBTBuilder::WithServiceLog(/*ParentNode*/ RootNode, /*ActivationIndex*/ INDEX_NONE, /*DeactivationIndex*/ INDEX_NONE, /*TickIndex*/ INDEX_NONE, /*TickBoolKeyName*/ TEXT("Bool3"), /*bCallTickOnSearchStart*/ false, /*BecomeRelevantBoolKeyName*/ NAME_None, /*BecomeRelevantBoolKeyName*/ NAME_None, /*bToggleValue*/ false, /*TicksDelaySetKeyNameTick*/ 2);
+
+		{
+			UBTCompositeNode& CompNode = FBTBuilder::AddSelector(RootNode);
+			FBTBuilder::WithDecoratorBlackboard(/*ParentNode*/ RootNode, /*Condition*/ EBasicKeyOperation::NotSet, /*Observer*/ EBTFlowAbortMode::LowerPriority, /*BoolKeyName*/  TEXT("Bool1"));
+			{
+				UBTCompositeNode& CompNode2 = FBTBuilder::AddSelector(CompNode);
+				FBTBuilder::WithDecoratorBlackboard(/*ParentNode*/ RootNode, /*Condition*/ EBasicKeyOperation::Set, /*Observer*/ EBTFlowAbortMode::LowerPriority, /*BoolKeyName*/ TEXT("Bool2"));
+				{
+					FBTBuilder::AddTaskLogFinish(/*ParentNode*/ CompNode2, /*LogIndex*/ IncorrectTaskExecute, /*FinishIndex*/ INDEX_NONE, EBTNodeResult::Succeeded, /*ExecutionTick*/ 0);
+				}
+			}
+		}
+
+		{
+			UBTCompositeNode& CompNode3 = FBTBuilder::AddSelector(RootNode);
+			FBTBuilder::WithDecoratorBlackboard(/*ParentNode*/ RootNode, /*Condition*/ EBasicKeyOperation::Set, /*Observer*/ EBTFlowAbortMode::LowerPriority, /*BoolKeyName*/ TEXT("Bool3"));
+
+			FBTBuilder::AddTaskLogFinish(/*ParentNode*/ CompNode3, /*LogIndex*/ SecondaryTask, /*FinishIndex*/ INDEX_NONE, EBTNodeResult::Succeeded, /*ExecutionTick*/ 4);
+		}
+
+		{
+			FBTBuilder::AddTaskLatentFlags(RootNode, EBTNodeResult::Succeeded, 3, FName(), MainTaskExecute, INDEX_NONE, 3, NAME_None, 2, INDEX_NONE);
+		}
+
+		ExpectedResult.Add(MainTaskExecute);
+		ExpectedResult.Add(StartAbortMainTask);
+		ExpectedResult.Add(SecondaryTask);
+	}
+};
+IMPLEMENT_AI_LATENT_TEST(FAITest_BTTestTickAuxNodesDeactivateNodes, "System.AI.Behavior Trees.Ticking aux nodes deactivates nodes")
+
+struct FAITest_BTTestAbortNestedSimpleParallel : public FAITest_SimpleBT
+{
+	FAITest_BTTestAbortNestedSimpleParallel()
+	{
+		enum
+		{
+			ChildFirstTaskExecute = 1,
+			ChildSecondTaskExecute,
+			ChildThirdTaskExecute
+		};
+
+		UBehaviorTree* ChildBTAsset = &FBTBuilder::CreateBehaviorTree();
+		AddAutoDestroyObject(*ChildBTAsset);
+
+		// Setup subtree
+		{
+			UBTCompositeNode& RootNode = FBTBuilder::AddSequence(*ChildBTAsset);
+			{
+				UBTCompositeNode& CompNode = FBTBuilder::AddParallel(RootNode, EBTParallelMode::AbortBackground);
+				{
+					FBTBuilder::AddTaskLogFinish(/*ParentNode*/ CompNode, /*LogIndex*/ ChildFirstTaskExecute, /*FinishIndex*/ INDEX_NONE, EBTNodeResult::Succeeded, /*ExecutionTick*/ 500);
+					FBTBuilder::WithDecorator<UBTDecorator_Cooldown>(CompNode);
+
+					UBTCompositeNode& CompNode2 = FBTBuilder::AddSequence(CompNode);
+					{
+						FBTBuilder::AddTaskLogFinish(/*ParentNode*/ CompNode2, /*LogIndex*/ ChildSecondTaskExecute, /*FinishIndex*/ INDEX_NONE, EBTNodeResult::Succeeded, /*ExecutionTick*/ 5);
+						FBTBuilder::WithDecorator<UBTDecorator_Cooldown>(CompNode2);
+
+						FBTBuilder::AddTaskLogFinish(/*ParentNode*/ CompNode2, /*LogIndex*/ ChildThirdTaskExecute, /*FinishIndex*/ INDEX_NONE, EBTNodeResult::Succeeded, /*ExecutionTick*/ 5);
+					}
+				}
+			}
+		}
+
+		// Setup main tree
+		{
+			UBTCompositeNode& RootNode = FBTBuilder::AddSequence(*BTAsset);
+			{
+				constexpr float NumTicks = 4.f;
+				UBTCompositeNode& CompNode = FBTBuilder::AddSelector(RootNode);
+				FBTBuilder::WithDecorator<UBTDecorator_TimeLimit>(RootNode).TimeLimit = NumTicks * FAITestHelpers::TickInterval;
+				{
+					FBTBuilder::AddTask<UTestBTTask_RunBehavior>(CompNode).SetSubtreeAsset(ChildBTAsset);
+				}
+			}
+		}
+
+		ExpectedResult.Add(ChildFirstTaskExecute);
+		ExpectedResult.Add(ChildSecondTaskExecute);
+	}
+};
+IMPLEMENT_AI_LATENT_TEST(FAITest_BTTestAbortNestedSimpleParallel, "System.AI.Behavior Trees.Abort: nested simple parallel")
+
 //////////////////////////////////////////////////////////////////////////////////////////////////
 // Stop tree tests
 //////////////////////////////////////////////////////////////////////////////////////////////////
@@ -2936,7 +3103,6 @@
 };
 IMPLEMENT_AI_LATENT_TEST(FAITest_BTStopDuringTaskTick, "System.AI.Behavior Trees.Stop: during task tick")
 
-
 struct FAITest_BTStopDuringTaskAbort : public FAITest_SimpleBT
 {
 	FAITest_BTStopDuringTaskAbort()
@@ -2986,7 +3152,6 @@
 };
 IMPLEMENT_AI_LATENT_TEST(FAITest_BTStopDuringTaskAbort, "System.AI.Behavior Trees.Stop: during task abort")
 
-
 struct FAITest_BTStopDuringTaskFinished : public FAITest_SimpleBT
 {
 	FAITest_BTStopDuringTaskFinished()
@@ -3283,7 +3448,6 @@
 };
 IMPLEMENT_AI_LATENT_TEST(FAITest_BTStopDuringTaskServiceBecomeRelevant, "System.AI.Behavior Trees.Stop: during task service become relevant")
 
-
 struct FAITest_BTStopDuringTaskServiceCeaseRelevant : public FAITest_SimpleBT
 {
 	FAITest_BTStopDuringTaskServiceCeaseRelevant()
@@ -3402,7 +3566,6 @@
 };
 IMPLEMENT_AI_LATENT_TEST(FAITest_BTUninitializeDuringTaskTick, "System.AI.Behavior Trees.Uninitialize: during task tick")
 
-
 struct FAITest_BTUninitializeDuringTaskAbort : public FAITest_SimpleBT
 {
 	FAITest_BTUninitializeDuringTaskAbort()
@@ -3452,7 +3615,6 @@
 };
 IMPLEMENT_AI_LATENT_TEST(FAITest_BTUninitializeDuringTaskAbort, "System.AI.Behavior Trees.Uninitialize: during task abort")
 
-
 struct FAITest_BTUninitializeDuringTaskFinished : public FAITest_SimpleBT
 {
 	FAITest_BTUninitializeDuringTaskFinished()
@@ -3749,10 +3911,8 @@
 };
 IMPLEMENT_AI_LATENT_TEST(FAITest_BTUninitializeDuringTaskServiceBecomeRelevant, "System.AI.Behavior Trees.Uninitialize: during task service become relevant")
 
-
 struct FAITest_BTUninitializeDuringTaskServiceCeaseRelevant : public FAITest_SimpleBT
-{
-	FAITest_BTUninitializeDuringTaskServiceCeaseRelevant()
+{	FAITest_BTUninitializeDuringTaskServiceCeaseRelevant()
 	{
 	    enum
 	    {
@@ -3799,7 +3959,6 @@
 };
 IMPLEMENT_AI_LATENT_TEST(FAITest_BTUninitializeDuringTaskServiceCeaseRelevant, "System.AI.Behavior Trees.Uninitialize: during task service cease relevant")
 
-
 //////////////////////////////////////////////////////////////////////////////////////////////////
 // Cleanup tests
 //////////////////////////////////////////////////////////////////////////////////////////////////
@@ -3869,7 +4028,6 @@
 };
 IMPLEMENT_AI_LATENT_TEST(FAITest_BTCleanupDuringTaskTick, "System.AI.Behavior Trees.Cleanup: during task tick")
 
-
 struct FAITest_BTCleanupDuringTaskAbort : public FAITest_SimpleBT
 {
 	FAITest_BTCleanupDuringTaskAbort()
@@ -3919,7 +4077,6 @@
 };
 IMPLEMENT_AI_LATENT_TEST(FAITest_BTCleanupDuringTaskAbort, "System.AI.Behavior Trees.Cleanup: during task abort")
 
-
 struct FAITest_BTCleanupDuringTaskFinished : public FAITest_SimpleBT
 {
 	FAITest_BTCleanupDuringTaskFinished()
@@ -4216,7 +4373,6 @@
 };
 IMPLEMENT_AI_LATENT_TEST(FAITest_BTCleanupDuringTaskServiceBecomeRelevant, "System.AI.Behavior Trees.Cleanup: during task service become relevant")
 
-
 struct FAITest_BTCleanupDuringTaskServiceCeaseRelevant : public FAITest_SimpleBT
 {
 	FAITest_BTCleanupDuringTaskServiceCeaseRelevant()
@@ -4366,137 +4522,4 @@
 };
 IMPLEMENT_AI_LATENT_TEST(FAITest_BTCompleteRestartDuringTaskExecute, "System.AI.Behavior Trees.Restart: complete during task execute")
 
-// This test is temporarily turned off until we enable CVarApplyAuxNodesFromFailedSearches again
-//struct FAITest_BTTestOberserverAbortLowerPriorityTwoDeep: public FAITest_SimpleBT
-//{
-//	FAITest_BTTestOberserverAbortLowerPriorityTwoDeep()
-//	{
-//		enum
-//		{
-//			MainTaskExecute = 1,
-//			InteruptingTaskExecute,
-//		};
-//
-//		UBTCompositeNode& RootNode = FBTBuilder::AddSelector(*BTAsset);
-//		UBTCompositeNode& CompNode = FBTBuilder::AddSelector(RootNode);
-//
-//		{
-//			UBTCompositeNode& Comp1Node = FBTBuilder::AddSelector(CompNode);
-//			FBTBuilder::WithDecoratorBlackboard(CompNode, EArithmeticKeyOperation::Equal, 1, EBTFlowAbortMode::LowerPriority, EBTBlackboardRestart::ValueChange, TEXT("Int"));
-//
-//			UBTCompositeNode& Comp2Node = FBTBuilder::AddSelector(Comp1Node);
-//			FBTBuilder::WithDecoratorBlackboard(Comp1Node, EArithmeticKeyOperation::Equal, 1, EBTFlowAbortMode::LowerPriority, EBTBlackboardRestart::ValueChange, TEXT("Int2"));
-//
-//			FBTBuilder::AddTask(Comp2Node, InteruptingTaskExecute, EBTNodeResult::Succeeded);
-//		}
-//
-//		{
-//			FBTBuilder::AddTaskValuesChangedWithLogs(/*ParentNode*/ CompNode, /*LogIndex*/ MainTaskExecute, /*NodeResult*/ EBTNodeResult::Succeeded, /*Value1*/1, /*Value2*/ 1, /*IntKeyName2*/ TEXT("Int"), /* IntKeyName2 */ TEXT("Int2"), /*ExecutionTicks1*/ 5, /*ExecutionTicks2*/ 5);
-//		}
-//
-//		ExpectedResult.Add(MainTaskExecute);
-//		ExpectedResult.Add(InteruptingTaskExecute);
-//	}
-//};
-//IMPLEMENT_AI_LATENT_TEST(FAITest_BTTestOberserverAbortLowerPriorityTwoDeep, "System.AI.Behavior Trees.OberserverAbortLowerPriorityTwoDeep: Two nodes in same branch with observer aborts lower priority")
-
-struct FAITest_BTTestTickAuxNodesDeactivateNodes : public FAITest_SimpleBT
-{
-	FAITest_BTTestTickAuxNodesDeactivateNodes()
-	{
-		enum
-		{
-			MainTaskExecute = 1,
-			StartAbortMainTask,
-			SecondaryTask,
-			IncorrectTaskExecute,
-		};
-
-		UBTCompositeNode& RootNode = FBTBuilder::AddSelector(*BTAsset);
-		FBTBuilder::WithServiceLog(/*ParentNode*/ RootNode, /*ActivationIndex*/ INDEX_NONE, /*DeactivationIndex*/ INDEX_NONE, /*TickIndex*/ INDEX_NONE, /*TickBoolKeyName*/ TEXT("Bool1"), /*bCallTickOnSearchStart*/ false, /*BecomeRelevantBoolKeyName*/ NAME_None, /*BecomeRelevantBoolKeyName*/ NAME_None, /*bToggleValue*/ false, /*TicksDelaySetKeyNameTick*/ 4);
-		FBTBuilder::WithServiceLog(/*ParentNode*/ RootNode, /*ActivationIndex*/ INDEX_NONE, /*DeactivationIndex*/ INDEX_NONE, /*TickIndex*/ INDEX_NONE, /*TickBoolKeyName*/ TEXT("Bool3"), /*bCallTickOnSearchStart*/ false, /*BecomeRelevantBoolKeyName*/ NAME_None, /*BecomeRelevantBoolKeyName*/ NAME_None, /*bToggleValue*/ false, /*TicksDelaySetKeyNameTick*/ 2);
-
-		{
-			UBTCompositeNode& CompNode = FBTBuilder::AddSelector(RootNode);
-			FBTBuilder::WithDecoratorBlackboard(/*ParentNode*/ RootNode, /*Condition*/ EBasicKeyOperation::NotSet, /*Observer*/ EBTFlowAbortMode::LowerPriority, /*BoolKeyName*/  TEXT("Bool1"));
-			{
-				UBTCompositeNode& CompNode2 = FBTBuilder::AddSelector(CompNode);
-				FBTBuilder::WithDecoratorBlackboard(/*ParentNode*/ RootNode, /*Condition*/ EBasicKeyOperation::Set, /*Observer*/ EBTFlowAbortMode::LowerPriority, /*BoolKeyName*/ TEXT("Bool2"));
-				{
-					FBTBuilder::AddTaskLogFinish(/*ParentNode*/ CompNode2, /*LogIndex*/ IncorrectTaskExecute, /*FinishIndex*/ INDEX_NONE, EBTNodeResult::Succeeded, /*ExecutionTick*/ 0);
-				}
-			}
-		}
-
-		{
-			UBTCompositeNode& CompNode3 = FBTBuilder::AddSelector(RootNode);
-			FBTBuilder::WithDecoratorBlackboard(/*ParentNode*/ RootNode, /*Condition*/ EBasicKeyOperation::Set, /*Observer*/ EBTFlowAbortMode::LowerPriority, /*BoolKeyName*/ TEXT("Bool3"));
-
-			FBTBuilder::AddTaskLogFinish(/*ParentNode*/ CompNode3, /*LogIndex*/ SecondaryTask, /*FinishIndex*/ INDEX_NONE, EBTNodeResult::Succeeded, /*ExecutionTick*/ 4);
-		}
-
-		{
-			FBTBuilder::AddTaskLatentFlags(RootNode, EBTNodeResult::Succeeded, 3, FName(), MainTaskExecute, INDEX_NONE, 3, NAME_None, 2, INDEX_NONE);
-		}
-
-		ExpectedResult.Add(MainTaskExecute);
-		ExpectedResult.Add(StartAbortMainTask);
-		ExpectedResult.Add(SecondaryTask);
-	}
-};
-IMPLEMENT_AI_LATENT_TEST(FAITest_BTTestTickAuxNodesDeactivateNodes, "System.AI.Behavior Trees.Ticking aux nodes deactivates nodes")
-
-struct FAITest_BTTestAbortNestedSimpleParallel : public FAITest_SimpleBT
-{
-	FAITest_BTTestAbortNestedSimpleParallel()
-	{
-		enum
-		{
-			ChildFirstTaskExecute = 1,
-			ChildSecondTaskExecute,
-			ChildThirdTaskExecute
-		};
-
-		UBehaviorTree* ChildBTAsset = &FBTBuilder::CreateBehaviorTree();
-		AddAutoDestroyObject(*ChildBTAsset);
-
-		// Setup subtree
-		{
-			UBTCompositeNode& RootNode = FBTBuilder::AddSequence(*ChildBTAsset);
-			{
-				UBTCompositeNode& CompNode = FBTBuilder::AddParallel(RootNode, EBTParallelMode::AbortBackground);
-				{
-					FBTBuilder::AddTaskLogFinish(/*ParentNode*/ CompNode, /*LogIndex*/ ChildFirstTaskExecute, /*FinishIndex*/ INDEX_NONE, EBTNodeResult::Succeeded, /*ExecutionTick*/ 500);
-					FBTBuilder::WithDecorator<UBTDecorator_Cooldown>(CompNode);
-
-					UBTCompositeNode& CompNode2 = FBTBuilder::AddSequence(CompNode);
-					{
-						FBTBuilder::AddTaskLogFinish(/*ParentNode*/ CompNode2, /*LogIndex*/ ChildSecondTaskExecute, /*FinishIndex*/ INDEX_NONE, EBTNodeResult::Succeeded, /*ExecutionTick*/ 5);
-						FBTBuilder::WithDecorator<UBTDecorator_Cooldown>(CompNode2);
-
-						FBTBuilder::AddTaskLogFinish(/*ParentNode*/ CompNode2, /*LogIndex*/ ChildThirdTaskExecute, /*FinishIndex*/ INDEX_NONE, EBTNodeResult::Succeeded, /*ExecutionTick*/ 5);
-					}
-				}
-			}
-		}
-
-		// Setup main tree
-		{
-			UBTCompositeNode& RootNode = FBTBuilder::AddSequence(*BTAsset);
-			{
-				constexpr float NumTicks = 4.f;
-				UBTCompositeNode& CompNode = FBTBuilder::AddSelector(RootNode);
-				FBTBuilder::WithDecorator<UBTDecorator_TimeLimit>(RootNode).TimeLimit = NumTicks * FAITestHelpers::TickInterval;
-				{
-					FBTBuilder::AddTask<UTestBTTask_RunBehavior>(CompNode).SetSubtreeAsset(ChildBTAsset);
-				}
-			}
-		}
-
-		ExpectedResult.Add(ChildFirstTaskExecute);
-		ExpectedResult.Add(ChildSecondTaskExecute);
-	}
-};
-IMPLEMENT_AI_LATENT_TEST(FAITest_BTTestAbortNestedSimpleParallel, "System.AI.Behavior Trees.Abort: nested simple parallel")
-
 #undef LOCTEXT_NAMESPACE
diff --git a/Engine/Source/Runtime/AIModule/Private/BehaviorTree/Decorators/BTDecorator_ForceSuccess.cpp b/Engine/Source/Runtime/AIModule/Private/BehaviorTree/Decorators/BTDecorator_ForceSuccess.cpp
--- a/Engine/Source/Runtime/AIModule/Private/BehaviorTree/Decorators/BTDecorator_ForceSuccess.cpp
+++ b/Engine/Source/Runtime/AIModule/Private/BehaviorTree/Decorators/BTDecorator_ForceSuccess.cpp
@@ -27,9 +27,10 @@ void UBTDecorator_ForceSuccess::OnNodeProcessed(FBehaviorTreeSearchData& SearchD
 	// from our parent node (abort self) or on our associated node (failure).
 	if (!SearchData.bSearchInProgress)
 	{
-		const FBTNodeIndex ParentNodeIndex(SearchData.RollbackInstanceIdx, GetParentNode() != nullptr ? GetParentNode()->GetExecutionIndex() : 0);
+		const int32 InstanceIdx = SearchData.OwnerComp.GetActiveInstanceIdx();
+		const FBTNodeIndex ParentNodeIndex(InstanceIdx, GetParentNode() != nullptr ? GetParentNode()->GetExecutionIndex() : 0);
 		const UBTNode* MyNode = GetMyNode();
-		const FBTNodeIndex MyNodeIndex(SearchData.RollbackInstanceIdx, MyNode != nullptr ? MyNode->GetExecutionIndex() : 0);
+		const FBTNodeIndex MyNodeIndex(InstanceIdx, MyNode != nullptr ? MyNode->GetExecutionIndex() : 0);
 		bCanForceSuccess = (SearchData.SearchRootNode == ParentNodeIndex || SearchData.SearchRootNode == MyNodeIndex);
 	}
 
