[AITestingSuite] fixed testableCheckf macro that was forcing function's `return` even if the rested condition didn't fail.

#rnx

diff --git a/Engine/Source/Developer/AITestSuite/Public/TestableEnsures.h b/Engine/Source/Developer/AITestSuite/Public/TestableEnsures.h
--- a/Engine/Source/Developer/AITestSuite/Public/TestableEnsures.h
+++ b/Engine/Source/Developer/AITestSuite/Public/TestableEnsures.h
@@ -40,8 +40,11 @@ namespace UE::AITestSuite
 #define testableCheckf(InExpression, InFormat, ... ) \
 	if (UNLIKELY(UE::AITestSuite::TestsInProgress > 0)) \
 	{ \
-		UE_CLOG(!(InExpression), LogTestableEnsures, Error, InFormat, ##__VA_ARGS__); \
-		return; \
+		if (!(InExpression))\
+		{ \
+			UE_LOG(LogTestableEnsures, Error, InFormat, ##__VA_ARGS__); \
+			return; \
+		}\
 	} \
 	else \
 	{ \
