Fix test run behavior task displaying in the BT task menu

diff --git a/Engine/Source/Developer/AITestSuite/Classes/BehaviorTree/TestBTTask_RunBehavior.h b/Engine/Source/Developer/AITestSuite/Classes/BehaviorTree/TestBTTask_RunBehavior.h
--- a/Engine/Source/Developer/AITestSuite/Classes/BehaviorTree/TestBTTask_RunBehavior.h
+++ b/Engine/Source/Developer/AITestSuite/Classes/BehaviorTree/TestBTTask_RunBehavior.h
@@ -5,7 +5,7 @@
 #include "BehaviorTree/Tasks/BTTask_RunBehavior.h"
 #include "TestBTTask_RunBehavior.generated.h"
 
-UCLASS(MinimalAPI)
+UCLASS(MinimalAPI, meta = (HiddenNode))
 class UTestBTTask_RunBehavior : public UBTTask_RunBehavior
 {
 	GENERATED_UCLASS_BODY()
