Update AITEST_SCOPED_CHECK to use `AddExpectedErrorPlain`.

`AddExpectedError` takes a regular expression. All usages of this macro were passing plain strings, not regular expressions.

This has the side benefit of being AutoRTFM-compatible, as the regular expression parser in AutomationTest currently uses ICU, which is not used in the server and so has not been wrapped for transactional safety.

diff --git a/Engine/Source/Developer/AITestSuite/Public/TestableEnsures.h b/Engine/Source/Developer/AITestSuite/Public/TestableEnsures.h
--- a/Engine/Source/Developer/AITestSuite/Public/TestableEnsures.h
+++ b/Engine/Source/Developer/AITestSuite/Public/TestableEnsures.h
@@ -16,7 +16,7 @@ namespace UE::AITestSuite
 }
 
 #define AITEST_SCOPED_CHECK(PartialTextToExpect, NumOccurrences) \
-	GetTestRunner().AddExpectedError(PartialTextToExpect, EAutomationExpectedErrorFlags::Contains, NumOccurrences); \
+	GetTestRunner().AddExpectedErrorPlain(PartialTextToExpect, EAutomationExpectedErrorFlags::Contains, NumOccurrences); \
 	ON_SCOPE_EXIT \
 	{ \
 		--UE::AITestSuite::TestsInProgress; \
