[AITest] Modified AI test suite test macros to support providing the test runner (useful to run tests inside a lambda)

diff --git a/Engine/Source/Developer/AITestSuite/Public/AITestsCommon.h b/Engine/Source/Developer/AITestSuite/Public/AITestsCommon.h
--- a/Engine/Source/Developer/AITestSuite/Public/AITestsCommon.h
+++ b/Engine/Source/Developer/AITestSuite/Public/AITestsCommon.h
@@ -210,20 +210,25 @@ struct FAITest_SimpleComponentBasedTest : public FAITestBase
 // fails making sure the rest of the test relying on given condition being 
 // true doesn't crash
 //----------------------------------------------------------------------//
-#define __AITEST_IMPL(What, Value, Test, RetVal)\
-	if (!GetTestRunner().Test(What, Value))\
+#define __AITEST_IMPL(_TestRunner, What, Value, Test, RetVal)\
+	if (!_TestRunner.Test(What, Value))\
 	{\
 		UE::AITestSuite::ConditionallyBreakOnTestFail(); \
 		return RetVal;\
 	}
 	
-#define AITEST_TRUE(What, Value) __AITEST_IMPL(What, Value, TestTrue, false)
-#define AITEST_FALSE(What, Value) __AITEST_IMPL(What, Value, TestFalse, false)
-#define AITEST_NULL(What, Pointer) __AITEST_IMPL(What, Pointer, TestNull, false)
-#define AITEST_NOT_NULL(What, Pointer) \
-	__AITEST_IMPL(What, Pointer, TestNotNull, false); \
+#define AITEST_TRUE_EX(_TestRunner, What, Value) __AITEST_IMPL(_TestRunner, What, Value, TestTrue, false)
+#define AITEST_FALSE_EX(_TestRunner, What, Value) __AITEST_IMPL(_TestRunner, What, Value, TestFalse, false)
+#define AITEST_NULL_EX(_TestRunner, What, Pointer) __AITEST_IMPL(_TestRunner, What, Pointer, TestNull, false)
+#define AITEST_NOT_NULL_EX(_TestRunner, What, Pointer) \
+	__AITEST_IMPL(_TestRunner, What, Pointer, TestNotNull, false); \
 	CA_ASSUME(Pointer)
 
+#define AITEST_TRUE(What, Value) AITEST_TRUE_EX(GetTestRunner(), What, Value)
+#define AITEST_FALSE(What, Value) AITEST_FALSE_EX(GetTestRunner(), What, Value)
+#define AITEST_NULL(What, Pointer) AITEST_NULL_EX(GetTestRunner(), What, Pointer)
+#define AITEST_NOT_NULL(What, Pointer) AITEST_NOT_NULL_EX(GetTestRunner(), What, Pointer)
+
 namespace FTestHelpers
 {
 	template<typename T1, typename T2>
@@ -256,17 +261,20 @@ namespace FTestHelpers
 	}
 }
 
-#define __AITEST_HELPER_IMPL(What, Actual, Expected, Test, RetVal)\
-	if (!FTestHelpers::Test(What, Actual, Expected, GetTestRunner()))\
+#define __AITEST_HELPER_IMPL(_TestRunner, What, Actual, Expected, Test, RetVal)\
+	if (!FTestHelpers::Test(What, Actual, Expected, _TestRunner))\
 	{\
 		UE::AITestSuite::ConditionallyBreakOnTestFail(); \
 		return RetVal;\
 	}
 
-#define AITEST_EQUAL(What, Actual, Expected) __AITEST_HELPER_IMPL(What, Actual, Expected, TestEqual, false)
-#define AITEST_NOT_EQUAL(What, Actual, Expected) __AITEST_HELPER_IMPL(What, Actual, Expected, TestNotEqual, false)
+#define AITEST_EQUAL_EX(_TestRunner, What, Actual, Expected) __AITEST_HELPER_IMPL(_TestRunner, What, Actual, Expected, TestEqual, false)
+#define AITEST_NOT_EQUAL_EX(_TestRunner, What, Actual, Expected) __AITEST_HELPER_IMPL(_TestRunner, What, Actual, Expected, TestNotEqual, false)
+#define AITEST_INFO_EX(_TestRunner, InfoString) FTestHelpers::AddInfo(InfoString, _TestRunner)
 
-#define AITEST_INFO(InfoString) FTestHelpers::AddInfo(InfoString, GetTestRunner())
+#define AITEST_EQUAL(What, Actual, Expected) AITEST_EQUAL_EX(GetTestRunner(), What, Actual, Expected)
+#define AITEST_NOT_EQUAL(What, Actual, Expected) AITEST_NOT_EQUAL_EX(GetTestRunner(), What, Actual, Expected)
+#define AITEST_INFO(InfoString) AITEST_INFO_EX(GetTestRunner(), InfoString)
 
 
 //----------------------------------------------------------------------//
@@ -278,9 +286,9 @@ namespace FTestHelpers
 // Update function is not a good place for testing results since its return
 // value controls whether the test will continue.
 //----------------------------------------------------------------------//
-#define AITEST_TRUE_LATENT(What, Value) __AITEST_IMPL(What, Value, TestTrue, )
-#define AITEST_FALSE_LATENT(What, Value) __AITEST_IMPL(What, Value, TestFalse, )
-#define AITEST_NULL_LATENT(What, Pointer) __AITEST_IMPL(What, Pointer, TestNull, )
-#define AITEST_NOT_NULL_LATENT(What, Pointer) __AITEST_IMPL(What, Pointer, TestNotNull, )
-#define AITEST_EQUAL_LATENT(What, Actual, Expected) __AITEST_HELPER_IMPL(What, Actual, Expected, TestEqual, )
-#define AITEST_NOT_EQUAL_LATENT(What, Actual, Expected) __AITEST_HELPER_IMPL(What, Actual, Expected, TestNotEqual, )
+#define AITEST_TRUE_LATENT(What, Value) __AITEST_IMPL(GetTestRunner(), What, Value, TestTrue, )
+#define AITEST_FALSE_LATENT(What, Value) __AITEST_IMPL(GetTestRunner(), What, Value, TestFalse, )
+#define AITEST_NULL_LATENT(What, Pointer) __AITEST_IMPL(GetTestRunner(), What, Pointer, TestNull, )
+#define AITEST_NOT_NULL_LATENT(What, Pointer) __AITEST_IMPL(GetTestRunner(), What, Pointer, TestNotNull, )
+#define AITEST_EQUAL_LATENT(What, Actual, Expected) __AITEST_HELPER_IMPL(GetTestRunner(), What, Actual, Expected, TestEqual, )
+#define AITEST_NOT_EQUAL_LATENT(What, Actual, Expected) __AITEST_HELPER_IMPL(GetTestRunner(), What, Actual, Expected, TestNotEqual, )
